CaptureManager: Move plugin to Virtual Production directory

diff --git a/Engine/Plugins/VirtualProduction/CaptureManager/CaptureManagerCore/Source/CaptureProtocolStack/Private/Tests/Discovery/TestDiscoveryNotify.cpp b/Engine/Plugins/VirtualProduction/CaptureManager/CaptureManagerCore/Source/CaptureProtocolStack/Private/Tests/Discovery/TestDiscoveryNotify.cpp
new file mode 100644
--- /dev/null
+++ b/Engine/Plugins/VirtualProduction/CaptureManager/CaptureManagerCore/Source/CaptureProtocolStack/Private/Tests/Discovery/TestDiscoveryNotify.cpp
@@ -0,0 +1,229 @@
+// Copyright Epic Games, Inc. All Rights Reserved.
+
+#include "Discovery/Messages/DiscoveryNotify.h"
+
+#include "Misc/AutomationTest.h"
+
+#if WITH_DEV_AUTOMATION_TESTS
+
+namespace UE::CaptureManager
+{
+
+IMPLEMENT_SIMPLE_AUTOMATION_TEST(TDiscoveryNotifyDeserializeTest, "Plugin.CaptureProtocolStack.Discovery.DiscoveryNotify.Deserialize.Success", EAutomationTestFlags_ApplicationContextMask | EAutomationTestFlags::SmokeFilter)
+bool TDiscoveryNotifyDeserializeTest::RunTest(const FString& InParameters)
+{
+	// Prepare test
+	const TArray<uint8> Header = { 'C', 'P', 'S', 'D', 'I', 'S', 'C', 'O', 'V', 'E', 'R', '\0' };
+	const FDiscoveryPacket::EMessageType MessageType = FDiscoveryPacket::EMessageType::Notify;
+	const TArray<uint8> ServerId = { 0x0d, 0x0e, 0x0a, 0x0d, 0x0b, 0x0e, 0x0e, 0x0f, 0x0d, 0x0e, 0x0a, 0x0d, 0x0b, 0x0e, 0x0e, 0x0f };
+	const uint16 ControlPort = 8000;
+	const FDiscoveryNotify::EConnectionState ConnectionState = FDiscoveryNotify::EConnectionState::Online;
+	const TArray<uint16> SupportedVersions = { 1, 2, 3 };
+
+	TArray<uint8> Packet;
+	Packet.Reserve(Header.Num() + sizeof(MessageType) + ServerId.Num() + sizeof(ControlPort) + SupportedVersions.Num() * sizeof(uint16));
+	Packet.Append(Header);
+	Packet.Add(static_cast<uint8>(MessageType));
+	Packet.Append(ServerId);
+	Packet.Append((uint8*) &ControlPort, sizeof(ControlPort));
+	Packet.Add(static_cast<uint8>(ConnectionState));
+	Packet.Append((uint8*) SupportedVersions.GetData(), SupportedVersions.Num() * sizeof(uint16));
+
+	// Deserialization
+	TProtocolResult<FDiscoveryPacket> PacketResult = FDiscoveryPacket::Deserialize(Packet);
+	TestTrue(TEXT("Packet Deserialize"), PacketResult.HasValue());
+
+	FDiscoveryPacket DiscoveryPacket = PacketResult.StealValue();
+	TProtocolResult<FDiscoveryNotify> NotifyResult = FDiscoveryNotify::Deserialize(DiscoveryPacket);
+	TestTrue(TEXT("Notify Deserialize"), NotifyResult.HasValue());
+
+	FDiscoveryNotify Notify = NotifyResult.StealValue();
+	int32 ServerIdCompareResult = FMemory::Memcmp(Notify.GetServerId().GetData(), ServerId.GetData(), ServerId.Num());
+	TestTrue(TEXT("Notify Deserialize"), (ServerIdCompareResult == 0));
+	TestEqual(TEXT("Notify Deserialize"), Notify.GetControlPort(), ControlPort);
+	TestEqual(TEXT("Notify Deserialize"), Notify.GetConnectionState(), ConnectionState);
+	TestEqual(TEXT("Notify Deserialize"), Notify.GetSupportedVersions(), SupportedVersions);
+
+	return true;
+}
+
+IMPLEMENT_SIMPLE_AUTOMATION_TEST(TDiscoveryNotifyDeserializeTestInvalidMessageType, "Plugin.CaptureProtocolStack.Discovery.DiscoveryNotify.Deserialize.InvalidMessageType", EAutomationTestFlags_ApplicationContextMask | EAutomationTestFlags::SmokeFilter)
+bool TDiscoveryNotifyDeserializeTestInvalidMessageType::RunTest(const FString& InParameters)
+{
+	// Prepare test
+	const TArray<uint8> Header = { 'C', 'P', 'S', 'D', 'I', 'S', 'C', 'O', 'V', 'E', 'R', '\0' };
+	const FDiscoveryPacket::EMessageType MessageType = FDiscoveryPacket::EMessageType::Request;
+
+	TArray<uint8> Packet;
+	Packet.Reserve(Header.Num() + sizeof(MessageType));
+	Packet.Append(Header);
+	Packet.Add(static_cast<uint8>(MessageType));
+
+	// Deserialization
+	TProtocolResult<FDiscoveryPacket> PacketResult = FDiscoveryPacket::Deserialize(Packet);
+	TestTrue(TEXT("Packet Deserialize"), PacketResult.HasValue());
+
+	FDiscoveryPacket DiscoveryPacket = PacketResult.StealValue();
+	TProtocolResult<FDiscoveryNotify> NotifyResult = FDiscoveryNotify::Deserialize(DiscoveryPacket);
+	TestTrue(TEXT("Notify Deserialize"), NotifyResult.HasError());
+
+	return true;
+}
+
+IMPLEMENT_SIMPLE_AUTOMATION_TEST(TDiscoveryNotifyDeserializeTestInvalidSize, "Plugin.CaptureProtocolStack.Discovery.DiscoveryNotify.Deserialize.InvalidSize", EAutomationTestFlags_ApplicationContextMask | EAutomationTestFlags::SmokeFilter)
+bool TDiscoveryNotifyDeserializeTestInvalidSize::RunTest(const FString& InParameters)
+{
+	// Prepare test
+	const TArray<uint8> Header = { 'C', 'P', 'S', 'D', 'I', 'S', 'C', 'O', 'V', 'E', 'R', '\0' };
+	const FDiscoveryPacket::EMessageType MessageType = FDiscoveryPacket::EMessageType::Notify;
+	const TArray<uint8> ServerId = { 0x0d, 0x0e, 0x0a, 0x0d, 0x0b, 0x0e, 0x0e, 0x0f };
+
+	TArray<uint8> Packet;
+	Packet.Reserve(Header.Num() + sizeof(MessageType));
+	Packet.Append(Header);
+	Packet.Add(static_cast<uint8>(MessageType));
+	Packet.Append(ServerId);
+
+	// Deserialization
+	TProtocolResult<FDiscoveryPacket> PacketResult = FDiscoveryPacket::Deserialize(Packet);
+	TestTrue(TEXT("Packet Deserialize"), PacketResult.HasValue());
+
+	FDiscoveryPacket DiscoveryPacket = PacketResult.StealValue();
+	TProtocolResult<FDiscoveryNotify> NotifyResult = FDiscoveryNotify::Deserialize(DiscoveryPacket);
+	TestTrue(TEXT("Notify Deserialize"), NotifyResult.HasError());
+
+	return true;
+}
+
+IMPLEMENT_SIMPLE_AUTOMATION_TEST(TDiscoveryNotifyDeserializeTestInvalidConnectionState, "Plugin.CaptureProtocolStack.Discovery.DiscoveryNotify.Deserialize.InvalidConnectionState", EAutomationTestFlags_ApplicationContextMask | EAutomationTestFlags::SmokeFilter)
+bool TDiscoveryNotifyDeserializeTestInvalidConnectionState::RunTest(const FString& InParameters)
+{
+	// Prepare test
+	const TArray<uint8> Header = { 'C', 'P', 'S', 'D', 'I', 'S', 'C', 'O', 'V', 'E', 'R', '\0' };
+	const FDiscoveryPacket::EMessageType MessageType = FDiscoveryPacket::EMessageType::Notify;
+	const TArray<uint8> ServerId = { 0x0d, 0x0e, 0x0a, 0x0d, 0x0b, 0x0e, 0x0e, 0x0f, 0x0d, 0x0e, 0x0a, 0x0d, 0x0b, 0x0e, 0x0e, 0x0f };
+	const uint16 ControlPort = 8000;
+	const FDiscoveryNotify::EConnectionState ConnectionState = FDiscoveryNotify::EConnectionState::Invalid;
+	const TArray<uint16> SupportedVersions = { 1, 2, 3 };
+
+	TArray<uint8> Packet;
+	Packet.Reserve(Header.Num() + sizeof(MessageType) + ServerId.Num() + sizeof(ControlPort) + SupportedVersions.Num() * sizeof(uint16));
+	Packet.Append(Header);
+	Packet.Add(static_cast<uint8>(MessageType));
+	Packet.Append(ServerId);
+	Packet.Append((uint8*) &ControlPort, sizeof(ControlPort));
+	Packet.Add(static_cast<uint8>(ConnectionState));
+	Packet.Append((uint8*) SupportedVersions.GetData(), SupportedVersions.Num() * sizeof(uint16));
+
+	// Deserialization
+	TProtocolResult<FDiscoveryPacket> PacketResult = FDiscoveryPacket::Deserialize(Packet);
+	TestTrue(TEXT("Packet Deserialize"), PacketResult.HasValue());
+
+	FDiscoveryPacket DiscoveryPacket = PacketResult.StealValue();
+	TProtocolResult<FDiscoveryNotify> NotifyResult = FDiscoveryNotify::Deserialize(DiscoveryPacket);
+	TestTrue(TEXT("Notify Deserialize"), NotifyResult.HasError());
+
+	return true;
+}
+
+IMPLEMENT_SIMPLE_AUTOMATION_TEST(TDiscoveryNotifyDeserializeTestInvalidSupportedVersionsSuccess, "Plugin.CaptureProtocolStack.Discovery.DiscoveryNotify.Deserialize.InvalidSupportedVersionsArray.Success", EAutomationTestFlags_ApplicationContextMask | EAutomationTestFlags::SmokeFilter)
+bool TDiscoveryNotifyDeserializeTestInvalidSupportedVersionsSuccess::RunTest(const FString& InParameters)
+{
+	// Prepare test
+	const TArray<uint8> Header = { 'C', 'P', 'S', 'D', 'I', 'S', 'C', 'O', 'V', 'E', 'R', '\0' };
+	const FDiscoveryPacket::EMessageType MessageType = FDiscoveryPacket::EMessageType::Notify;
+	const TArray<uint8> ServerId = { 0x0d, 0x0e, 0x0a, 0x0d, 0x0b, 0x0e, 0x0e, 0x0f, 0x0d, 0x0e, 0x0a, 0x0d, 0x0b, 0x0e, 0x0e, 0x0f };
+	const uint16 ControlPort = 8000;
+	const FDiscoveryNotify::EConnectionState ConnectionState = FDiscoveryNotify::EConnectionState::Online;
+	const TArray<uint8> SupportedVersions = { 0x01, 0x00, 0x02, 0x00, 0x03 };
+	check(SupportedVersions.Num() % 2 != 0);
+
+	TArray<uint8> Packet;
+	Packet.Reserve(Header.Num() + sizeof(MessageType) + ServerId.Num() + sizeof(ControlPort) + SupportedVersions.Num());
+	Packet.Append(Header);
+	Packet.Add(static_cast<uint8>(MessageType));
+	Packet.Append(ServerId);
+	Packet.Append((uint8*) &ControlPort, sizeof(ControlPort));
+	Packet.Add(static_cast<uint8>(ConnectionState));
+	Packet.Append(SupportedVersions.GetData(), SupportedVersions.Num());
+
+	// Deserialization
+	TProtocolResult<FDiscoveryPacket> PacketResult = FDiscoveryPacket::Deserialize(Packet);
+	TestTrue(TEXT("Packet Deserialize"), PacketResult.HasValue());
+
+	FDiscoveryPacket DiscoveryPacket = PacketResult.StealValue();
+	TProtocolResult<FDiscoveryNotify> NotifyResult = FDiscoveryNotify::Deserialize(DiscoveryPacket);
+	TestTrue(TEXT("Notify Deserialize"), NotifyResult.HasValue());
+
+	return true;
+}
+
+IMPLEMENT_SIMPLE_AUTOMATION_TEST(TDiscoveryNotifyDeserializeTestInvalidSupportedVersionsFailure, "Plugin.CaptureProtocolStack.Discovery.DiscoveryNotify.Deserialize.InvalidSupportedVersionsArray.Failure", EAutomationTestFlags_ApplicationContextMask | EAutomationTestFlags::SmokeFilter)
+bool TDiscoveryNotifyDeserializeTestInvalidSupportedVersionsFailure::RunTest(const FString& InParameters)
+{
+	// Prepare test
+	const TArray<uint8> Header = { 'C', 'P', 'S', 'D', 'I', 'S', 'C', 'O', 'V', 'E', 'R', '\0' };
+	const FDiscoveryPacket::EMessageType MessageType = FDiscoveryPacket::EMessageType::Notify;
+	const TArray<uint8> ServerId = { 0x0d, 0x0e, 0x0a, 0x0d, 0x0b, 0x0e, 0x0e, 0x0f, 0x0d, 0x0e, 0x0a, 0x0d, 0x0b, 0x0e, 0x0e, 0x0f };
+	const uint16 ControlPort = 8000;
+	const FDiscoveryNotify::EConnectionState ConnectionState = FDiscoveryNotify::EConnectionState::Online;
+	const TArray<uint8> SupportedVersions = { 0x01 };
+	check(SupportedVersions.Num() % 2 != 0);
+
+	TArray<uint8> Packet;
+	Packet.Reserve(Header.Num() + sizeof(MessageType) + ServerId.Num() + sizeof(ControlPort) + SupportedVersions.Num());
+	Packet.Append(Header);
+	Packet.Add(static_cast<uint8>(MessageType));
+	Packet.Append(ServerId);
+	Packet.Append((uint8*) &ControlPort, sizeof(ControlPort));
+	Packet.Add(static_cast<uint8>(ConnectionState));
+	Packet.Append(SupportedVersions.GetData(), SupportedVersions.Num());
+
+	// Deserialization
+	TProtocolResult<FDiscoveryPacket> PacketResult = FDiscoveryPacket::Deserialize(Packet);
+	TestTrue(TEXT("Packet Deserialize"), PacketResult.HasValue());
+
+	FDiscoveryPacket DiscoveryPacket = PacketResult.StealValue();
+	TProtocolResult<FDiscoveryNotify> NotifyResult = FDiscoveryNotify::Deserialize(DiscoveryPacket);
+	TestTrue(TEXT("Notify Deserialize"), NotifyResult.HasError());
+
+	return true;
+}
+
+IMPLEMENT_SIMPLE_AUTOMATION_TEST(TDiscoveryNotifySerializeTest, "Plugin.CaptureProtocolStack.Discovery.DiscoveryNotify.Serialize.Success", EAutomationTestFlags_ApplicationContextMask | EAutomationTestFlags::SmokeFilter)
+bool TDiscoveryNotifySerializeTest::RunTest(const FString& InParameters)
+{
+	// Prepare test
+	const FDiscoveryPacket::EMessageType MessageType = FDiscoveryPacket::EMessageType::Notify;
+
+	const TArray<uint8> ServerId = { 0x0d, 0x0e, 0x0a, 0x0d, 0x0b, 0x0e, 0x0e, 0x0f, 0x0d, 0x0e, 0x0a, 0x0d, 0x0b, 0x0e, 0x0e, 0x0f };
+	const uint16 ControlPort = 8000;
+	const FDiscoveryNotify::EConnectionState ConnectionState = FDiscoveryNotify::EConnectionState::Online;
+	const TArray<uint16> SupportedVersions = { 1 };
+
+	TArray<uint8> Payload;
+	Payload.Reserve(ServerId.Num() + sizeof(ControlPort) + sizeof(ConnectionState) + SupportedVersions.Num() * sizeof(uint16));
+	Payload.Append(ServerId);
+	Payload.Append((uint8*) &ControlPort, sizeof(ControlPort));
+	Payload.Add(static_cast<uint8>(ConnectionState));
+	Payload.Append((uint8*) SupportedVersions.GetData(), SupportedVersions.Num() * sizeof(uint16));
+
+	FDiscoveryNotify::FServerId ServerIdStatic;
+	FMemory::Memcpy(ServerIdStatic.GetData(), ServerId.GetData(), ServerId.Num());
+
+	FDiscoveryNotify Notify(MoveTemp(ServerIdStatic), ControlPort, ConnectionState, SupportedVersions);
+
+	// Deserialization
+	TProtocolResult<FDiscoveryPacket> PacketResult = FDiscoveryNotify::Serialize(Notify);
+	TestTrue(TEXT("Packet Serialize"), PacketResult.HasValue());
+
+	FDiscoveryPacket DiscoveryPacket = PacketResult.StealValue();
+	TestEqual(TEXT("Notify Serialize"), DiscoveryPacket.GetMessageType(), MessageType);
+	TestEqual(TEXT("Notify Serialize"), DiscoveryPacket.GetPayload(), Payload);
+
+	return true;
+}
+
+}
+
+#endif // WITH_DEV_AUTOMATION_TESTS
\ No newline at end of file
diff --git a/Engine/Plugins/VirtualProduction/CaptureManager/CaptureManagerCore/Source/CaptureProtocolStack/Private/Tests/Discovery/TestDiscoveryPacket.cpp b/Engine/Plugins/VirtualProduction/CaptureManager/CaptureManagerCore/Source/CaptureProtocolStack/Private/Tests/Discovery/TestDiscoveryPacket.cpp
new file mode 100644
--- /dev/null
+++ b/Engine/Plugins/VirtualProduction/CaptureManager/CaptureManagerCore/Source/CaptureProtocolStack/Private/Tests/Discovery/TestDiscoveryPacket.cpp
@@ -0,0 +1,121 @@
+// Copyright Epic Games, Inc. All Rights Reserved.
+
+#include "Discovery/Communication/DiscoveryPacket.h"
+
+#include "Misc/AutomationTest.h"
+
+#if WITH_DEV_AUTOMATION_TESTS
+
+namespace UE::CaptureManager
+{
+
+IMPLEMENT_SIMPLE_AUTOMATION_TEST(TDiscoveryPacketDeserializeTest, "Plugin.CaptureProtocolStack.Discovery.DiscoveryPacket.Deserialize.Success", EAutomationTestFlags_ApplicationContextMask | EAutomationTestFlags::SmokeFilter)
+bool TDiscoveryPacketDeserializeTest::RunTest(const FString& InParameters)
+{
+	// Prepare test
+	const TArray<uint8> Header = { 'C', 'P', 'S', 'D', 'I', 'S', 'C', 'O', 'V', 'E', 'R', '\0' };
+	const FDiscoveryPacket::EMessageType MessageType = FDiscoveryPacket::EMessageType::Request;
+	const TArray<uint8> Payload = { 0x0d, 0x0e, 0x0a, 0x0d, 0x0b, 0x0e, 0x0e, 0x0f, 0x00 };
+
+	TArray<uint8> Packet;
+	Packet.Reserve(Header.Num() + sizeof(MessageType) + Payload.Num());
+	Packet.Append(Header);
+	Packet.Add(static_cast<uint8>(MessageType));
+	Packet.Append(Payload);
+
+	// Deserialization
+	TProtocolResult<FDiscoveryPacket> PacketResult = FDiscoveryPacket::Deserialize(Packet);
+	TestTrue(TEXT("Packet Deserialize"), PacketResult.HasValue());
+
+	FDiscoveryPacket DiscoveryPacket = PacketResult.StealValue();
+	TestEqual(TEXT("Packet Deserialize"), DiscoveryPacket.GetMessageType(), MessageType);
+	TestEqual(TEXT("Packet Deserialize"), DiscoveryPacket.GetPayload().Num(), Payload.Num());
+	TestEqual(TEXT("Packet Deserialize"), DiscoveryPacket.GetPayload(), Payload);
+
+	return true;
+}
+
+IMPLEMENT_SIMPLE_AUTOMATION_TEST(TDiscoveryPacketDeserializeTestInvalidHeaderSize, "Plugin.CaptureProtocolStack.Discovery.DiscoveryPacket.Deserialize.InvalidHeaderSize", EAutomationTestFlags_ApplicationContextMask | EAutomationTestFlags::SmokeFilter)
+bool TDiscoveryPacketDeserializeTestInvalidHeaderSize::RunTest(const FString& InParameters)
+{
+	// Prepare test
+	const TArray<uint8> Header = { 'A', 'A', 'A', '\0' };
+	const FDiscoveryPacket::EMessageType MessageType = FDiscoveryPacket::EMessageType::Request;
+
+	TArray<uint8> Packet;
+	Packet.Reserve(Header.Num() + sizeof(MessageType));
+	Packet.Append(Header);
+	Packet.Add(static_cast<uint8>(MessageType));
+
+	// Deserialization
+	TProtocolResult<FDiscoveryPacket> PacketResult = FDiscoveryPacket::Deserialize(Packet);
+	TestTrue(TEXT("Packet Deserialize"), PacketResult.HasError());
+
+	return true;
+}
+
+IMPLEMENT_SIMPLE_AUTOMATION_TEST(TDiscoveryPacketDeserializeTestInvalidHeader, "Plugin.CaptureProtocolStack.Discovery.DiscoveryPacket.Deserialize.InvalidHeader", EAutomationTestFlags_ApplicationContextMask | EAutomationTestFlags::SmokeFilter)
+bool TDiscoveryPacketDeserializeTestInvalidHeader::RunTest(const FString& InParameters)
+{
+	// Prepare test
+	const TArray<uint8> Header = { 'A', 'A', 'A', 'A', 'A', 'A', 'A', 'A', 'A', 'A', 'A', '\0' };
+	const FDiscoveryPacket::EMessageType MessageType = FDiscoveryPacket::EMessageType::Request;
+
+	TArray<uint8> Packet;
+	Packet.Reserve(Header.Num() + sizeof(MessageType));
+	Packet.Append(Header);
+	Packet.Add(static_cast<uint8>(MessageType));
+
+	// Deserialization
+	TProtocolResult<FDiscoveryPacket> PacketResult = FDiscoveryPacket::Deserialize(Packet);
+	TestTrue(TEXT("Packet Deserialize"), PacketResult.HasError());
+
+	return true;
+}
+
+IMPLEMENT_SIMPLE_AUTOMATION_TEST(TDiscoveryPacketDeserializeTestInvalidMessageType, "Plugin.CaptureProtocolStack.Discovery.DiscoveryPacket.Deserialize.InvalidMessageType", EAutomationTestFlags_ApplicationContextMask | EAutomationTestFlags::SmokeFilter)
+bool TDiscoveryPacketDeserializeTestInvalidMessageType::RunTest(const FString& InParameters)
+{
+	// Prepare test
+	const TArray<uint8> Header = { 'C', 'P', 'S', 'D', 'I', 'S', 'C', 'O', 'V', 'E', 'R', '\0' };
+	const FDiscoveryPacket::EMessageType MessageType = FDiscoveryPacket::EMessageType::Invalid;
+
+	TArray<uint8> Packet;
+	Packet.Reserve(Header.Num() + sizeof(MessageType));
+	Packet.Append(Header);
+	Packet.Add(static_cast<uint8>(MessageType));
+
+	// Deserialization
+	TProtocolResult<FDiscoveryPacket> PacketResult = FDiscoveryPacket::Deserialize(Packet);
+	TestTrue(TEXT("Packet Deserialize"), PacketResult.HasError());
+
+	return true;
+}
+
+IMPLEMENT_SIMPLE_AUTOMATION_TEST(TDiscoveryPacketSerializeTest, "Plugin.CaptureProtocolStack.Discovery.DiscoveryPacket.Serialize.Success", EAutomationTestFlags_ApplicationContextMask | EAutomationTestFlags::SmokeFilter)
+bool TDiscoveryPacketSerializeTest::RunTest(const FString& InParameters)
+{
+	// Prepare test
+	const TArray<uint8> Header = { 'C', 'P', 'S', 'D', 'I', 'S', 'C', 'O', 'V', 'E', 'R', '\0' };
+	const FDiscoveryPacket::EMessageType MessageType = FDiscoveryPacket::EMessageType::Request;
+
+	TArray<uint8> Packet;
+	Packet.Reserve(Header.Num() + sizeof(MessageType));
+	Packet.Append(Header);
+	Packet.Add(static_cast<uint8>(MessageType));
+
+	FDiscoveryPacket DiscoveryPacket(MessageType, TArray<uint8>());
+
+	// Deserialization
+	TProtocolResult<TArray<uint8>> DataResult = FDiscoveryPacket::Serialize(DiscoveryPacket);
+	TestTrue(TEXT("Packet Serialize"), DataResult.HasValue());
+
+	TArray<uint8> Data = DataResult.StealValue();
+	TestEqual(TEXT("Packet Serialize"), Packet, Data);
+
+	return true;
+}
+
+}
+
+#endif //WITH_DEV_AUTOMATION_TESTS
\ No newline at end of file
diff --git a/Engine/Plugins/VirtualProduction/CaptureManager/CaptureManagerCore/Source/CaptureProtocolStack/Private/Tests/Discovery/TestDiscoveryResponse.cpp b/Engine/Plugins/VirtualProduction/CaptureManager/CaptureManagerCore/Source/CaptureProtocolStack/Private/Tests/Discovery/TestDiscoveryResponse.cpp
new file mode 100644
--- /dev/null
+++ b/Engine/Plugins/VirtualProduction/CaptureManager/CaptureManagerCore/Source/CaptureProtocolStack/Private/Tests/Discovery/TestDiscoveryResponse.cpp
@@ -0,0 +1,189 @@
+// Copyright Epic Games, Inc. All Rights Reserved.
+
+#include "Discovery/Messages/DiscoveryResponse.h"
+
+#include "Misc/AutomationTest.h"
+
+#if WITH_DEV_AUTOMATION_TESTS
+
+namespace UE::CaptureManager
+{
+
+IMPLEMENT_SIMPLE_AUTOMATION_TEST(TDiscoveryResponseDeserializeTest, "Plugin.CaptureProtocolStack.Discovery.DiscoveryResponse.Deserialize.Success", EAutomationTestFlags_ApplicationContextMask | EAutomationTestFlags::SmokeFilter)
+bool TDiscoveryResponseDeserializeTest::RunTest(const FString& InParameters)
+{
+	// Prepare test
+	const TArray<uint8> Header = { 'C', 'P', 'S', 'D', 'I', 'S', 'C', 'O', 'V', 'E', 'R', '\0' };
+	const FDiscoveryPacket::EMessageType MessageType = FDiscoveryPacket::EMessageType::Response;
+	const TArray<uint8> ServerId = { 0x0d, 0x0e, 0x0a, 0x0d, 0x0b, 0x0e, 0x0e, 0x0f, 0x0d, 0x0e, 0x0a, 0x0d, 0x0b, 0x0e, 0x0e, 0x0f };
+	const uint16 ControlPort = 8000;
+	const TArray<uint16> SupportedVersions = { 1, 2, 3 };
+
+	TArray<uint8> Packet;
+	Packet.Reserve(Header.Num() + sizeof(MessageType) + ServerId.Num() + sizeof(ControlPort) + SupportedVersions.Num() * sizeof(uint16));
+	Packet.Append(Header);
+	Packet.Add(static_cast<uint8>(MessageType));
+	Packet.Append(ServerId);
+	Packet.Append((uint8*) &ControlPort, sizeof(ControlPort));
+	Packet.Append((uint8*) SupportedVersions.GetData(), SupportedVersions.Num() * sizeof(uint16));
+
+	// Deserialization
+	TProtocolResult<FDiscoveryPacket> PacketResult = FDiscoveryPacket::Deserialize(Packet);
+	TestTrue(TEXT("Packet Deserialize"), PacketResult.HasValue());
+
+	FDiscoveryPacket DiscoveryPacket = PacketResult.StealValue();
+	TProtocolResult<FDiscoveryResponse> ResponseResult = FDiscoveryResponse::Deserialize(DiscoveryPacket);
+	TestTrue(TEXT("Response Deserialize"), ResponseResult.HasValue());
+
+	FDiscoveryResponse Response = ResponseResult.StealValue();
+	int32 ServerIdCompareResult = FMemory::Memcmp(Response.GetServerId().GetData(), ServerId.GetData(), ServerId.Num());
+	TestTrue(TEXT("Response Deserialize"), (ServerIdCompareResult == 0));
+	TestEqual(TEXT("Response Deserialize"), Response.GetControlPort(), ControlPort);
+	TestEqual(TEXT("Response Deserialize"), Response.GetSupportedVersions(), SupportedVersions);
+
+	return true;
+}
+
+IMPLEMENT_SIMPLE_AUTOMATION_TEST(TDiscoveryResponseDeserializeTestInvalidMessageType, "Plugin.CaptureProtocolStack.Discovery.DiscoveryResponse.Deserialize.InvalidMessageType", EAutomationTestFlags_ApplicationContextMask | EAutomationTestFlags::SmokeFilter)
+bool TDiscoveryResponseDeserializeTestInvalidMessageType::RunTest(const FString& InParameters)
+{
+	// Prepare test
+	const TArray<uint8> Header = { 'C', 'P', 'S', 'D', 'I', 'S', 'C', 'O', 'V', 'E', 'R', '\0' };
+	const FDiscoveryPacket::EMessageType MessageType = FDiscoveryPacket::EMessageType::Request;
+
+	TArray<uint8> Packet;
+	Packet.Reserve(Header.Num() + sizeof(MessageType));
+	Packet.Append(Header);
+	Packet.Add(static_cast<uint8>(MessageType));
+
+	// Deserialization
+	TProtocolResult<FDiscoveryPacket> PacketResult = FDiscoveryPacket::Deserialize(Packet);
+	TestTrue(TEXT("Packet Deserialize"), PacketResult.HasValue());
+
+	FDiscoveryPacket DiscoveryPacket = PacketResult.StealValue();
+	TProtocolResult<FDiscoveryResponse> ResponseResult = FDiscoveryResponse::Deserialize(DiscoveryPacket);
+	TestTrue(TEXT("Response Deserialize"), ResponseResult.HasError());
+
+	return true;
+}
+
+IMPLEMENT_SIMPLE_AUTOMATION_TEST(TDiscoveryResponseDeserializeTestInvalidSize, "Plugin.CaptureProtocolStack.Discovery.DiscoveryResponse.Deserialize.InvalidSize", EAutomationTestFlags_ApplicationContextMask | EAutomationTestFlags::SmokeFilter)
+bool TDiscoveryResponseDeserializeTestInvalidSize::RunTest(const FString& InParameters)
+{
+	// Prepare test
+	const TArray<uint8> Header = { 'C', 'P', 'S', 'D', 'I', 'S', 'C', 'O', 'V', 'E', 'R', '\0' };
+	const FDiscoveryPacket::EMessageType MessageType = FDiscoveryPacket::EMessageType::Response;
+	const TArray<uint8> ServerId = { 0x0d, 0x0e, 0x0a, 0x0d, 0x0b, 0x0e, 0x0e, 0x0f };
+
+	TArray<uint8> Packet;
+	Packet.Reserve(Header.Num() + sizeof(MessageType));
+	Packet.Append(Header);
+	Packet.Add(static_cast<uint8>(MessageType));
+	Packet.Append(ServerId);
+
+	// Deserialization
+	TProtocolResult<FDiscoveryPacket> PacketResult = FDiscoveryPacket::Deserialize(Packet);
+	TestTrue(TEXT("Packet Deserialize"), PacketResult.HasValue());
+
+	FDiscoveryPacket DiscoveryPacket = PacketResult.StealValue();
+	TProtocolResult<FDiscoveryResponse> ResponseResult = FDiscoveryResponse::Deserialize(DiscoveryPacket);
+	TestTrue(TEXT("Response Deserialize"), ResponseResult.HasError());
+
+	return true;
+}
+
+IMPLEMENT_SIMPLE_AUTOMATION_TEST(TDiscoveryResponseDeserializeTestInvalidSupportedVersionsSuccess, "Plugin.CaptureProtocolStack.Discovery.DiscoveryResponse.Deserialize.InvalidSupportedVersionsArray.Success", EAutomationTestFlags_ApplicationContextMask | EAutomationTestFlags::SmokeFilter)
+bool TDiscoveryResponseDeserializeTestInvalidSupportedVersionsSuccess::RunTest(const FString& InParameters)
+{
+	// Prepare test
+	const TArray<uint8> Header = { 'C', 'P', 'S', 'D', 'I', 'S', 'C', 'O', 'V', 'E', 'R', '\0' };
+	const FDiscoveryPacket::EMessageType MessageType = FDiscoveryPacket::EMessageType::Response;
+	const TArray<uint8> ServerId = { 0x0d, 0x0e, 0x0a, 0x0d, 0x0b, 0x0e, 0x0e, 0x0f, 0x0d, 0x0e, 0x0a, 0x0d, 0x0b, 0x0e, 0x0e, 0x0f };
+	const uint16 ControlPort = 8000;
+	const TArray<uint8> SupportedVersions = { 0x01, 0x00, 0x02, 0x00, 0x03 };
+	check(SupportedVersions.Num() % 2 != 0);
+
+	TArray<uint8> Packet;
+	Packet.Reserve(Header.Num() + sizeof(MessageType) + ServerId.Num() + sizeof(ControlPort) + SupportedVersions.Num());
+	Packet.Append(Header);
+	Packet.Add(static_cast<uint8>(MessageType));
+	Packet.Append(ServerId);
+	Packet.Append((uint8*) &ControlPort, sizeof(ControlPort));
+	Packet.Append(SupportedVersions.GetData(), SupportedVersions.Num());
+
+	// Deserialization
+	TProtocolResult<FDiscoveryPacket> PacketResult = FDiscoveryPacket::Deserialize(Packet);
+	TestTrue(TEXT("Packet Deserialize"), PacketResult.HasValue());
+
+	FDiscoveryPacket DiscoveryPacket = PacketResult.StealValue();
+	TProtocolResult<FDiscoveryResponse> ResponseResult = FDiscoveryResponse::Deserialize(DiscoveryPacket);
+	TestTrue(TEXT("Notify Deserialize"), ResponseResult.HasValue());
+
+	return true;
+}
+
+IMPLEMENT_SIMPLE_AUTOMATION_TEST(TDiscoveryResponseDeserializeTestInvalidSupportedVersionsFailure, "Plugin.CaptureProtocolStack.Discovery.DiscoveryResponse.Deserialize.InvalidSupportedVersionsArray.Failure", EAutomationTestFlags_ApplicationContextMask | EAutomationTestFlags::SmokeFilter)
+bool TDiscoveryResponseDeserializeTestInvalidSupportedVersionsFailure::RunTest(const FString& InParameters)
+{
+	// Prepare test
+	const TArray<uint8> Header = { 'C', 'P', 'S', 'D', 'I', 'S', 'C', 'O', 'V', 'E', 'R', '\0' };
+	const FDiscoveryPacket::EMessageType MessageType = FDiscoveryPacket::EMessageType::Response;
+	const TArray<uint8> ServerId = { 0x0d, 0x0e, 0x0a, 0x0d, 0x0b, 0x0e, 0x0e, 0x0f, 0x0d, 0x0e, 0x0a, 0x0d, 0x0b, 0x0e, 0x0e, 0x0f };
+	const uint16 ControlPort = 8000;
+	const TArray<uint8> SupportedVersions = { 0x01 };
+	check(SupportedVersions.Num() % 2 != 0);
+
+	TArray<uint8> Packet;
+	Packet.Reserve(Header.Num() + sizeof(MessageType) + ServerId.Num() + sizeof(ControlPort) + SupportedVersions.Num());
+	Packet.Append(Header);
+	Packet.Add(static_cast<uint8>(MessageType));
+	Packet.Append(ServerId);
+	Packet.Append((uint8*) &ControlPort, sizeof(ControlPort));
+	Packet.Append(SupportedVersions.GetData(), SupportedVersions.Num());
+
+	// Deserialization
+	TProtocolResult<FDiscoveryPacket> PacketResult = FDiscoveryPacket::Deserialize(Packet);
+	TestTrue(TEXT("Packet Deserialize"), PacketResult.HasValue());
+
+	FDiscoveryPacket DiscoveryPacket = PacketResult.StealValue();
+	TProtocolResult<FDiscoveryResponse> ResponseResult = FDiscoveryResponse::Deserialize(DiscoveryPacket);
+	TestTrue(TEXT("Notify Deserialize"), ResponseResult.HasError());
+
+	return true;
+}
+
+IMPLEMENT_SIMPLE_AUTOMATION_TEST(TDiscoveryResponseSerializeTest, "Plugin.CaptureProtocolStack.Discovery.DiscoveryResponse.Serialize.Success", EAutomationTestFlags_ApplicationContextMask | EAutomationTestFlags::SmokeFilter)
+bool TDiscoveryResponseSerializeTest::RunTest(const FString& InParameters)
+{
+	// Prepare test
+	const FDiscoveryPacket::EMessageType MessageType = FDiscoveryPacket::EMessageType::Response;
+
+	const TArray<uint8> ServerId = { 0x0d, 0x0e, 0x0a, 0x0d, 0x0b, 0x0e, 0x0e, 0x0f, 0x0d, 0x0e, 0x0a, 0x0d, 0x0b, 0x0e, 0x0e, 0x0f };
+	const uint16 ControlPort = 8000;
+	const TArray<uint16> SupportedVersions = { 1 };
+
+	TArray<uint8> Payload;
+	Payload.Reserve(ServerId.Num() + sizeof(ControlPort) + SupportedVersions.Num() * sizeof(uint16));
+	Payload.Append(ServerId);
+	Payload.Append((uint8*) &ControlPort, sizeof(ControlPort));
+	Payload.Append((uint8*) SupportedVersions.GetData(), SupportedVersions.Num() * sizeof(uint16));
+
+	FDiscoveryResponse::FServerId ServerIdStatic;
+	FMemory::Memcpy(ServerIdStatic.GetData(), ServerId.GetData(), ServerId.Num());
+
+	FDiscoveryResponse Response(MoveTemp(ServerIdStatic), ControlPort, SupportedVersions);
+
+	// Deserialization
+	TProtocolResult<FDiscoveryPacket> PacketResult = FDiscoveryResponse::Serialize(Response);
+	TestTrue(TEXT("Packet Serialize"), PacketResult.HasValue());
+
+	FDiscoveryPacket DiscoveryPacket = PacketResult.StealValue();
+	TestEqual(TEXT("Response Serialize"), DiscoveryPacket.GetMessageType(), MessageType);
+	TestEqual(TEXT("Response Serialize"), DiscoveryPacket.GetPayload(), Payload);
+
+	return true;
+}
+
+}
+
+#endif // WITH_DEV_AUTOMATION_TESTS
\ No newline at end of file
diff --git a/Engine/Plugins/VirtualProduction/CaptureManager/CaptureManagerCore/Source/CaptureProtocolStack/Public/Discovery/DiscoveryMessenger.h b/Engine/Plugins/VirtualProduction/CaptureManager/CaptureManagerCore/Source/CaptureProtocolStack/Public/Discovery/DiscoveryMessenger.h
new file mode 100644
--- /dev/null
+++ b/Engine/Plugins/VirtualProduction/CaptureManager/CaptureManagerCore/Source/CaptureProtocolStack/Public/Discovery/DiscoveryMessenger.h
@@ -0,0 +1,44 @@
+// Copyright Epic Games, Inc. All Rights Reserved.
+
+#pragma once
+
+#include "Discovery/Communication/DiscoveryCommunication.h"
+#include "Discovery/Communication/DiscoveryPacket.h"
+
+#include "Discovery/Messages/DiscoveryRequest.h"
+#include "Discovery/Messages/DiscoveryResponse.h"
+#include "Discovery/Messages/DiscoveryNotify.h"
+
+#include "Utility/Error.h"
+
+namespace UE::CaptureManager
+{
+
+class CAPTUREPROTOCOLSTACK_API FDiscoveryMessenger final
+{
+public:
+	using FOnResponseArrived = TDelegate<void(FString, FDiscoveryResponse), FDefaultTSDelegateUserPolicy>;
+	using FOnNotifyArrived = TDelegate<void(FString, FDiscoveryNotify), FDefaultTSDelegateUserPolicy>;
+
+	FDiscoveryMessenger();
+	~FDiscoveryMessenger();
+
+	TProtocolResult<void> Start();
+	TProtocolResult<void> Stop();
+
+	TProtocolResult<void> SendMulticastRequest();
+
+	void SetResponseHandler(FOnResponseArrived InOnResponse);
+	void SetNotifyHandler(FOnNotifyArrived InOnNotify);
+
+private:
+
+	void OnPacketArrived(FString InServerIp, FDiscoveryPacket InPacket);
+
+	FDiscoveryCommunication Communication;
+
+	FOnResponseArrived OnResponse;
+	FOnNotifyArrived OnNotify;
+};
+
+}
\ No newline at end of file
diff --git a/Engine/Plugins/VirtualProduction/CaptureManager/CaptureManagerCore/Source/CaptureSourceFramework/Private/Capabilities/IngestCapability.cpp b/Engine/Plugins/VirtualProduction/CaptureManager/CaptureManagerCore/Source/CaptureSourceFramework/Private/Capabilities/IngestCapability.cpp
new file mode 100644
--- /dev/null
+++ b/Engine/Plugins/VirtualProduction/CaptureManager/CaptureManagerCore/Source/CaptureSourceFramework/Private/Capabilities/IngestCapability.cpp
@@ -0,0 +1,211 @@
+// Copyright Epic Games, Inc. All Rights Reserved.
+
+#include "Capabilities/IngestCapability.h"
+
+namespace UE::CaptureManager
+{
+
+FIngestCapabilityError::FIngestCapabilityError(ECode InCode, FString InMessage)
+	: Code(InCode)
+	, Message(MoveTemp(InMessage))
+{
+}
+
+FIngestCapabilityError::ECode FIngestCapabilityError::GetCode() const
+{
+	return Code;
+}
+
+const FString& FIngestCapabilityError::GetMessage() const
+{
+	return Message;
+}
+
+const FString FIngestCapabilityBase::Name = TEXT("Ingest");
+const FString FIngestCapabilityBase::CancelCmd = TEXT("Cancel");
+
+const FString FIngestCapabilityBase::Id = TEXT("Id");
+
+FIngestCapabilityBase::FIngestCapabilityBase()
+	: FCaptureSourceCapability(Name)
+{
+	AddCommand(FCommandDesc(FImportTakeCmd::Name,
+							{
+								FPropertyDesc(Id, FPropertyDesc::EType::Number),
+							}));
+	AddCommand(FCommandDesc(FUpdateTakeListCmd::Name, {}));
+
+	AddCommand(FCommandDesc(CancelCmd,
+							{
+								FPropertyDesc(Id, FPropertyDesc::EType::Number)
+							}));
+
+	AddEvent(FEventDesc(FTakeAddedEvent::Name,
+						{
+							FPropertyDesc(Id, FPropertyDesc::EType::Number)
+						}));
+	AddEvent(FEventDesc(FTakeRemovedEvent::Name,
+						{
+							FPropertyDesc(Id, FPropertyDesc::EType::Number)
+						}));
+	AddEvent(FEventDesc(FTakeUpdatedEvent::Name,
+			 {
+				 FPropertyDesc(Id, FPropertyDesc::EType::Number)
+			 }));
+}
+
+FIngestCapabilityBase::FTakeProcessHandle::FTakeProcessHandle(FIngestCapabilityBase* InOwner, FTakeId InId)
+	: Owner(InOwner)
+	, Id(InId)
+{
+}
+
+FTakeId FIngestCapabilityBase::AddTakeMetadata(FTakeMetadata InTakeInfo)
+{
+	FScopeLock Lock(&Mutex);
+
+	FTakeId TakeId = CurrentTakeId++;
+	Takes.Emplace(TakeId, MoveTemp(InTakeInfo));
+
+	return TakeId;
+}
+
+void FIngestCapabilityBase::RemoveTakeMetadata(FTakeId InTakeId)
+{
+	FScopeLock Lock(&Mutex);
+	Takes.Remove(InTakeId);
+}
+
+void FIngestCapabilityBase::UpdateTakeMetadata(FTakeId InTakeId, FTakeMetadata InTakeInfo)
+{
+	FScopeLock Lock(&Mutex);
+	if (FTakeMetadata* TakeMetadata = Takes.Find(InTakeId))
+	{
+		*TakeMetadata = MoveTemp(InTakeInfo);
+	}
+}
+
+TOptional<FTakeMetadata> FIngestCapabilityBase::GetTakeMetadata(FTakeId InTakeId) const
+{
+	FScopeLock Lock(&Mutex);
+	if (const FTakeMetadata* TakeMetadata = Takes.Find(InTakeId))
+	{
+		return *TakeMetadata;
+	}
+
+	return {};
+}
+
+TArray<FTakeId> FIngestCapabilityBase::GetTakeIds() const
+{
+	FScopeLock Lock(&Mutex);
+
+	TArray<FTakeId> TakeIds;
+	Takes.GetKeys(TakeIds);
+
+	return TakeIds;
+}
+
+void FIngestCapabilityBase::PublishTakeAddedEvent(FTakeId InTakeId)
+{
+	PublishEvent<FTakeAddedEvent>(InTakeId);
+}
+
+void FIngestCapabilityBase::PublishTakeUpdatedEvent(FTakeId InTakeId)
+{
+	PublishEvent<FTakeUpdatedEvent>(InTakeId);
+}
+
+void FIngestCapabilityBase::PublishTakeRemovedEvent(FTakeId InTakeId)
+{
+	PublishEvent<FTakeRemovedEvent>(InTakeId);
+}
+
+void FIngestCapabilityBase::FTakeProcessHandle::Cancel()
+{
+	Owner->Cancel(Id);
+}
+
+const FString FTakeAddedEvent::Name = TEXT("TakeAdded");
+const FString FTakeUpdatedEvent::Name = TEXT("TakeUpdated");
+const FString FTakeRemovedEvent::Name = TEXT("TakeRemoved");
+
+const FString FImportTakeCmd::Name = TEXT("ImportTake");
+FImportTakeCmd::FImportTakeCmd(FIngestOptions InIngestOptions,
+							   FProcessProgressReporter InProgressReporter,
+							   FProcessFinished InProcessFinished)
+	: FCommandBase(Name)
+	, IngestOptions(MoveTemp(InIngestOptions))
+	, ProgressReporter(MoveTemp(InProgressReporter))
+	, ProcessFinished(MoveTemp(InProcessFinished))
+{
+}
+
+const FString FUpdateTakeListCmd::Name = TEXT("UpdateTakeList");
+FUpdateTakeListCmd::FUpdateTakeListCmd(FCallback InCallback)
+	: FCommandBase(Name)
+	, Callback(MoveTemp(InCallback))
+{
+}
+
+FIngestCapability::FIngestCapability()
+{
+}
+
+FIngestCapability::FTakeProcessHandle FIngestCapability::ImportTakeAsync(FTakeId InTakeId,
+																		 FIngestOptions InIngestOptions,
+																		 FImportTakeCmd::FProcessFinished InProcessFinished,
+																		 FImportTakeCmd::FProcessProgressReporter InProcessProgressReported)
+{
+	TSharedPtr<FImportTakeCmd> InCommand = MakeShared<FImportTakeCmd>(MoveTemp(InIngestOptions),
+																	  MoveTemp(InProcessProgressReported),
+																	  MoveTemp(InProcessFinished));
+
+	InCommand->AddParamValue(Id, MakePropertyValue<int64>(InTakeId));
+
+	AsyncTask(ENamedThreads::AnyThread, [this, Command = MoveTemp(InCommand)]() mutable
+	{
+		ExecuteCommand(MoveTemp(Command));
+	});
+
+	return FTakeProcessHandle(this, InTakeId);
+}
+
+void FIngestCapability::UpdateTakeListAsync(FUpdateTakeListCmd::FCallback InUpdateTakeListCallback)
+{
+	{
+		FScopeLock Lock(&Mutex);
+		Takes.Empty();
+	}
+
+	AsyncTask(ENamedThreads::AnyThread, [this, Callback = MoveTemp(InUpdateTakeListCallback)]() mutable
+	{
+		ExecuteCommand(MakeShared<FUpdateTakeListCmd>(MoveTemp(Callback)));
+	});
+}
+
+void FIngestCapability::Cancel(FTakeId InTakeId)
+{
+	TSharedPtr<FCommandBase> InCommand = MakeShared<FCommandBase>(CancelCmd);
+
+	InCommand->AddParamValue(Id, MakePropertyValue<int64>(InTakeId));
+
+	ExecuteCommand(MoveTemp(InCommand));
+}
+
+void FIngestCapability::SetImportTakeHandler(FCommandHandler InCommandHandler)
+{
+	SetCommandHandler(FImportTakeCmd::Name, MoveTemp(InCommandHandler));
+}
+
+void FIngestCapability::SetUpdateTakeListHandler(FCommandHandler InCommandHandler)
+{
+	SetCommandHandler(FUpdateTakeListCmd::Name, MoveTemp(InCommandHandler));
+}
+
+void FIngestCapability::SetCancelHandler(FCommandHandler InCommandHandler)
+{
+	SetCommandHandler(CancelCmd, MoveTemp(InCommandHandler));
+}
+
+}
\ No newline at end of file
diff --git a/Engine/Plugins/VirtualProduction/CaptureManager/CaptureManagerCore/Source/CaptureSourceFramework/Private/CaptureSourceDetailsCustomization.cpp b/Engine/Plugins/VirtualProduction/CaptureManager/CaptureManagerCore/Source/CaptureSourceFramework/Private/CaptureSourceDetailsCustomization.cpp
new file mode 100644
--- /dev/null
+++ b/Engine/Plugins/VirtualProduction/CaptureManager/CaptureManagerCore/Source/CaptureSourceFramework/Private/CaptureSourceDetailsCustomization.cpp
@@ -0,0 +1,181 @@
+// Copyright Epic Games, Inc. All Rights Reserved.
+
+#include "CaptureSourceDetailsCustomization.h"
+
+#include "CaptureSourceFrameworkModule.h"
+
+#include "Widgets/Input/SButton.h"
+#include "Widgets/Input/SEditableTextBox.h"
+#include "Widgets/Input/SCheckBox.h"
+#include "Widgets/Input/SNumericEntryBox.h"
+#include "Widgets/SNullWidget.h"
+
+#include "Widgets/SDefaultCaptureSourceCreator.h"
+#include "Widgets/SDefaultCaptureSourceCapabilityCreator.h"
+#include "Widgets/SDefaultCaptureSourcePropertyCreator.h"
+#include "Widgets/SDefaultCaptureSourceCommandCreator.h"
+
+namespace UE::CaptureManager
+{
+
+template <typename ... Args>
+static FString GenerateKey(Args... InArgs)
+{
+	static_assert(sizeof...(Args) != 0 && sizeof...(Args) <= 3,
+				  "Method currently supports 1 - 3 arguments");
+
+	FString Format;
+	if constexpr (sizeof...(Args) == 1)
+	{
+		Format = TEXT("{0}");
+	}
+	else if constexpr (sizeof...(Args) == 2)
+	{
+		Format = TEXT("{0}.{1}");
+	}
+	else
+	{
+		Format = TEXT("{0}.{1}.{2}");
+	}
+	return FString::Format(*Format, { (InArgs)... });
+}
+
+FCaptureSourceDetailsCustomization::FCaptureSourceDetailsCustomization()
+{
+}
+
+void FCaptureSourceDetailsCustomization::RegisterPropertyCustomization(const FString& InCaptureSourceClassId,
+																	   const FString& InCapabilityName,
+																	   const FString& InPropertyName,
+																	   FPropertyWidgetCreator InWidgetCreator)
+{
+	PropertyCustomizations.Emplace(GenerateKey(InCaptureSourceClassId, InCapabilityName, InPropertyName), MoveTemp(InWidgetCreator));
+}
+
+void FCaptureSourceDetailsCustomization::RegisterCommandCustomization(const FString& InCaptureSourceClassId,
+																	  const FString& InCapabilityName,
+																	  const FString& InCommandName,
+																	  FCommandWidgetCreator InWidgetCreator)
+{
+	CommandCustomizations.Emplace(GenerateKey(InCaptureSourceClassId, InCapabilityName, InCommandName), MoveTemp(InWidgetCreator));
+}
+
+void FCaptureSourceDetailsCustomization::RegisterCapabilityCustomization(const FString& InCaptureSourceClassId,
+																		 const FString& InCapabilityName,
+																		 FCapabilityWidgetCreator InWidgetCreator)
+{
+	CapabilityCustomizations.Emplace(GenerateKey(InCaptureSourceClassId, InCapabilityName), MoveTemp(InWidgetCreator));
+}
+
+void FCaptureSourceDetailsCustomization::RegisterCaptureSourceCustomization(const FString& InCaptureSourceClassId,
+																			FCaptureSourceWidgetCreator InWidgetCreator)
+{
+	CaptureSourceCustomizations.Emplace(GenerateKey(InCaptureSourceClassId), MoveTemp(InWidgetCreator));
+}
+
+TSharedPtr<SWidget> FCaptureSourceDetailsCustomization::CreatePropertyWidget(const FString& InCaptureSourceClassId,
+																			 FCaptureSourceCapability* InCapability,
+																			 const FPropertyDesc& InProperty)
+{
+	if (!InCapability)
+	{
+		return SNullWidget::NullWidget;
+	}
+
+	FPropertyWidgetCreator* FoundDelegate = PropertyCustomizations.Find(GenerateKey(InCaptureSourceClassId, InCapability->GetName(), InProperty.Name));
+	if (FoundDelegate && FoundDelegate->IsBound())
+	{
+		TSharedPtr<SWidget> Widget = FoundDelegate->Execute(InCaptureSourceClassId, InCapability, InProperty);
+		return Widget;
+	}
+
+	return SNew(SDefaultCaptureSourcePropertyCreator)
+		.Capability(InCapability)
+		.Name(InProperty.Name);
+}
+
+TSharedPtr<SWidget> FCaptureSourceDetailsCustomization::CreateCommandWidget(const FString& InCaptureSourceClassId,
+																			FCaptureSourceCapability* InCapability,
+																			const FCommandDesc& InCommand)
+{
+	if (!InCapability)
+	{
+		return SNullWidget::NullWidget;
+	}
+
+	FCommandWidgetCreator* FoundDelegate = CommandCustomizations.Find(GenerateKey(InCaptureSourceClassId, InCapability->GetName(), InCommand.Name));
+	if (FoundDelegate && FoundDelegate->IsBound())
+	{
+		TSharedPtr<SWidget> Widget = FoundDelegate->Execute(InCaptureSourceClassId, InCapability, InCommand);
+		return Widget;
+	}
+
+	return SNew(SDefaultCaptureSourceCommandCreator)
+		.Capability(InCapability)
+		.Name(InCommand.Name);
+}
+
+TSharedPtr<SWidget> FCaptureSourceDetailsCustomization::CreateCapabilityWidget(const FString& InCaptureSourceClassId,
+																			   FCaptureSourceCapability* InCapability)
+{
+	if (!InCapability)
+	{
+		return SNullWidget::NullWidget;
+	}
+
+	FCapabilityWidgetCreator* FoundDelegate = CapabilityCustomizations.Find(GenerateKey(InCaptureSourceClassId, InCapability->GetName()));
+	if (FoundDelegate && FoundDelegate->IsBound())
+	{
+		TSharedPtr<SWidget> Widget = FoundDelegate->Execute(InCaptureSourceClassId, InCapability);
+		return Widget;
+	}
+
+	return SNew(SDefaultCaptureSourceCapabilityCreator)
+		.FactoryId(InCaptureSourceClassId)
+		.Capability(InCapability);
+}
+
+TSharedPtr<SWidget> FCaptureSourceDetailsCustomization::CreateCaptureSourceWidget(const FCaptureSourceDescriptor& InCaptureSourceDescriptor,
+																				  FOnCaptureSourceConfigured InCaptureSourceConfigured)
+{
+	FCaptureSourceWidgetCreator* FoundDelegate = CaptureSourceCustomizations.Find(InCaptureSourceDescriptor.GetId());
+	if (FoundDelegate && FoundDelegate->IsBound())
+	{
+		TSharedPtr<SWidget> Widget = FoundDelegate->Execute(InCaptureSourceDescriptor, MoveTemp(InCaptureSourceConfigured));
+		return Widget;
+	}
+
+	TSharedPtr<SDefaultCaptureSourceCreator> Widget = SNew(SDefaultCaptureSourceCreator)
+		.FactoryId(InCaptureSourceDescriptor.GetId())
+		.CreationParams(InCaptureSourceDescriptor.GetCreationParams())
+		.OnCaptureSourceConfigured(MoveTemp(InCaptureSourceConfigured));
+
+	return Widget;
+}
+
+bool FCaptureSourceDetailsCustomization::ContainsPropertyWidgetCustomization(const FString& InFactoryId,
+																			 const FString& InCapabilityName,
+																			 const FString& InPropertyName) const
+{
+	return PropertyCustomizations.Contains(GenerateKey(InFactoryId, InCapabilityName, InPropertyName));
+}
+
+bool FCaptureSourceDetailsCustomization::ContainsCommandWidgetCustomization(const FString& InFactoryId,
+																			const FString& InCapabilityName,
+																			const FString& InCommandName) const
+{
+	return CommandCustomizations.Contains(GenerateKey(InFactoryId, InCapabilityName, InCommandName));
+}
+
+bool FCaptureSourceDetailsCustomization::ContainsCapabilityWidgetCustomization(const FString& InFactoryId,
+																			   const FString& InCapabilityName) const
+{
+	return CapabilityCustomizations.Contains(GenerateKey(InFactoryId, InCapabilityName));
+}
+
+bool FCaptureSourceDetailsCustomization::ContainsCaptureSourceWidgetCustomization(const FString& InFactoryId) const
+{
+	return CaptureSourceCustomizations.Contains(GenerateKey(InFactoryId));
+}
+
+}
\ No newline at end of file
diff --git a/Engine/Plugins/VirtualProduction/CaptureManager/CaptureManagerCore/Source/CaptureSourceFramework/Private/CaptureSourceFactory.cpp b/Engine/Plugins/VirtualProduction/CaptureManager/CaptureManagerCore/Source/CaptureSourceFramework/Private/CaptureSourceFactory.cpp
new file mode 100644
--- /dev/null
+++ b/Engine/Plugins/VirtualProduction/CaptureManager/CaptureManagerCore/Source/CaptureSourceFramework/Private/CaptureSourceFactory.cpp
@@ -0,0 +1,78 @@
+// Copyright Epic Games, Inc. All Rights Reserved.
+
+#include "CaptureSourceFactory.h"
+
+FCaptureSourceDescriptor::FCaptureSourceDescriptor(const FString& InId, FText InSourceTypeDisplayName, TArray<FPropertyDesc> InCreationParams) :
+	Id(InId),
+	SourceTypeDisplayName(InSourceTypeDisplayName),
+	CreationParams(MoveTemp(InCreationParams))
+{
+}
+
+void FCaptureSourceDescriptor::SetDiscoverable(bool bInIsDiscoverable)
+{
+	bIsDiscoverable = bInIsDiscoverable;
+}
+
+bool FCaptureSourceDescriptor::IsDiscoverable() const
+{
+	return bIsDiscoverable;
+}
+
+void FCaptureSourceDescriptor::AddCreationParam(FPropertyDesc InParam)
+{
+	CreationParams.Push(MoveTemp(InParam));
+}
+
+const TArray<FPropertyDesc>& FCaptureSourceDescriptor::GetCreationParams() const
+{
+	return CreationParams;
+}
+
+const FString& FCaptureSourceDescriptor::GetId() const
+{
+	return Id;
+}
+
+const FText& FCaptureSourceDescriptor::GetSourceTypeDisplayName() const
+{
+	return SourceTypeDisplayName;
+}
+
+FCaptureVoidResult FCaptureSourceFactory::Discover(FDiscoverCallback InCallback)
+{
+	if (!IsDiscoverable())
+	{
+		InCallback.ExecuteIfBound(TArray<FCaptureSourceResult>());
+
+		return MakeError(FCaptureSourceError(ECaptureSourceError::NotSupported));
+	}
+
+	FCaptureVoidResult Result = MakeValue();
+
+	DiscoverImpl(MoveTemp(InCallback), Result);
+
+	return Result;
+}
+
+FCaptureSourceFactory::FCaptureSourceResult FCaptureSourceFactory::CreateCaptureSource(const FString& InName, TMap<FString, FPropertyValue> InCreationParamsValue)
+{
+	FCaptureSourceResult CaptureSourceResult = CreateCaptureSourceImpl(InName, InCreationParamsValue);
+
+	if (CaptureSourceResult.IsValid())
+	{
+		CaptureSourceResult.GetValue()->SetCreationParams(MoveTemp(InCreationParamsValue));
+	}
+
+	return CaptureSourceResult;
+}
+
+FCaptureSourceDescriptor FCaptureSourceFactory::CreateCaptureSourceDescriptor() const
+{
+	TArray<FPropertyDesc> MandatoryCreationParams;
+	FCaptureSourceDescriptor Descriptor(GetCaptureSourceFactoryId(), GetCaptureSourceFactoryName(), MoveTemp(MandatoryCreationParams));
+
+	CreateCaptureSourceDescriptorImpl(Descriptor);
+
+	return Descriptor;
+}
\ No newline at end of file
diff --git a/Engine/Plugins/VirtualProduction/CaptureManager/CaptureManagerCore/Source/CaptureSourceFramework/Private/Widgets/SDefaultCaptureSourceCapabilityCreator.cpp b/Engine/Plugins/VirtualProduction/CaptureManager/CaptureManagerCore/Source/CaptureSourceFramework/Private/Widgets/SDefaultCaptureSourceCapabilityCreator.cpp
new file mode 100644
--- /dev/null
+++ b/Engine/Plugins/VirtualProduction/CaptureManager/CaptureManagerCore/Source/CaptureSourceFramework/Private/Widgets/SDefaultCaptureSourceCapabilityCreator.cpp
@@ -0,0 +1,106 @@
+// Copyright Epic Games, Inc. All Rights Reserved.
+
+#include "Widgets/SDefaultCaptureSourceCapabilityCreator.h"
+
+#include "Widgets/SDefaultCaptureSourcePropertyCreator.h"
+#include "Widgets/SDefaultCaptureSourceCommandCreator.h"
+
+#include "WidgetUtils.h"
+
+#include "CaptureSourceFrameworkModule.h"
+
+#include "Widgets/SBoxPanel.h"
+#include "Widgets/Text/STextBlock.h"
+
+namespace UE::CaptureManager
+{
+
+void SDefaultCaptureSourceCapabilityCreator::Construct(const FArguments& InArgs)
+{
+	FactoryId = InArgs._FactoryId;
+	Capability = InArgs._Capability;
+
+	FCaptureSourceDetailsCustomization& Customization =
+		FModuleManager::Get().GetModuleChecked<FCaptureSourceFrameworkModule>("CaptureSourceFramework").GetCustomization();
+
+	TSharedPtr<SVerticalBox> VerticalBox = SNew(SVerticalBox);
+
+	for (const FPropertyDesc& Property : Capability->GetProperties())
+	{
+		if (Customization.ContainsPropertyWidgetCustomization(FactoryId, Capability->GetName(), Property.Name))
+		{
+			VerticalBox->AddSlot()
+				.AutoHeight()
+				.VAlign(EVerticalAlignment::VAlign_Center)
+				[
+					Customization.CreatePropertyWidget(FactoryId, Capability, Property).ToSharedRef()
+				];
+		}
+		else
+		{
+			FPropertyValue& Value = PropertyValues.Emplace(Property.Name);
+
+			TSharedRef<SDefaultCaptureSourcePropertyCreator> PropertyWidget = SNew(SDefaultCaptureSourcePropertyCreator)
+				.Capability(Capability)
+				.Name(Property.Name)
+				.ValuePtr(&Value);
+
+			VerticalBox->AddSlot()
+				.AutoHeight()
+				.VAlign(EVerticalAlignment::VAlign_Center)
+				[
+					PropertyWidget
+				];
+		}
+	}
+
+	Capability->SubscribeToEvent(
+		FCapturePropertyChangedEvent::EventName,
+		FCaptureEventHandler::Type::CreateRaw(this, &SDefaultCaptureSourceCapabilityCreator::HandlePropertyChangedEvent));
+
+	for (const FCommandDesc& Command : Capability->GetCommands())
+	{
+		TSharedPtr<SVerticalBox> CommandVerticalBox = SNew(SVerticalBox);
+
+		CommandVerticalBox->AddSlot()
+			.AutoHeight()
+			.VAlign(EVerticalAlignment::VAlign_Center)
+			[
+				SNew(STextBlock)
+					.Text(FText::FromString(Command.Name))
+			];
+
+		CommandVerticalBox->AddSlot()
+			.AutoHeight()
+			.VAlign(EVerticalAlignment::VAlign_Center)
+			[
+				Customization.CreateCommandWidget(FactoryId, Capability, Command).ToSharedRef()
+			];
+
+		VerticalBox->AddSlot()
+			.AutoHeight()
+			.VAlign(EVerticalAlignment::VAlign_Center)
+			[
+				CommandVerticalBox.ToSharedRef()
+			];
+	}
+
+	ChildSlot
+		[
+			VerticalBox.ToSharedRef()
+		];
+}
+
+void SDefaultCaptureSourceCapabilityCreator::HandlePropertyChangedEvent(TSharedPtr<const FCaptureEvent> InEvent)
+{
+	TSharedPtr<const FCapturePropertyChangedEvent> Event = StaticCastSharedPtr<const FCapturePropertyChangedEvent>(InEvent);
+
+	if (!PropertyValues.Contains(Event->Name))
+	{
+		return;
+	}
+
+	PropertyValues[Event->Name] = Event->Value;
+}
+
+}
\ No newline at end of file
diff --git a/Engine/Plugins/VirtualProduction/CaptureManager/CaptureManagerCore/Source/CaptureSourceFramework/Private/Widgets/SDefaultCaptureSourceCommandCreator.cpp b/Engine/Plugins/VirtualProduction/CaptureManager/CaptureManagerCore/Source/CaptureSourceFramework/Private/Widgets/SDefaultCaptureSourceCommandCreator.cpp
new file mode 100644
--- /dev/null
+++ b/Engine/Plugins/VirtualProduction/CaptureManager/CaptureManagerCore/Source/CaptureSourceFramework/Private/Widgets/SDefaultCaptureSourceCommandCreator.cpp
@@ -0,0 +1,74 @@
+// Copyright Epic Games, Inc. All Rights Reserved.
+
+#include "Widgets/SDefaultCaptureSourceCommandCreator.h"
+
+#include "WidgetUtils.h"
+
+#include "Widgets/SBoxPanel.h"
+#include "Widgets/Input/SButton.h"
+
+namespace UE::CaptureManager
+{
+
+void SDefaultCaptureSourceCommandCreator::Construct(const FArguments& InArgs)
+{
+	Capability = InArgs._Capability;
+	Name = InArgs._Name;
+
+	TSharedPtr<SHorizontalBox> HorizontalBox = SNew(SHorizontalBox);
+
+	FCommandDesc CommandDesc = FindCommandDesc();
+
+	TSharedPtr<SVerticalBox> VerticalBox = SNew(SVerticalBox);
+
+	for (const FPropertyDesc& Property : CommandDesc.Params)
+	{
+		Params.Emplace(Property.Name, UE::CaptureManager::Private::GetDefaultValueForProperty(Property));
+
+		TSharedRef<SWidget> Widget = UE::CaptureManager::Private::CreateParamWidget(Property, Params[Property.Name]).ToSharedRef();
+
+		VerticalBox->AddSlot()
+			.AutoHeight()
+			.VAlign(EVerticalAlignment::VAlign_Center)
+			.AttachWidget(MoveTemp(Widget));
+	}
+
+	VerticalBox->AddSlot()
+		.AutoHeight()
+		.VAlign(EVerticalAlignment::VAlign_Bottom)
+		.HAlign(EHorizontalAlignment::HAlign_Right)
+		[
+			SNew(SButton)
+				.Text(FText::FromString(TEXT("Execute")))
+				.OnClicked_Lambda([this]() mutable
+			{
+				TSharedPtr<FCommandBase> CommandPtr = MakeShared<FCommandBase>(Name);
+				CommandPtr->SetParamValues(Params);
+
+				Capability->ExecuteCommand(MoveTemp(CommandPtr));
+
+				return FReply::Handled();
+			})
+		];
+
+	ChildSlot
+		[
+			VerticalBox.ToSharedRef()
+		];
+}
+
+FCommandDesc SDefaultCaptureSourceCommandCreator::FindCommandDesc()
+{
+	TArray<FCommandDesc> Commands = Capability->GetCommands();
+
+	FCommandDesc* Found = Commands.FindByPredicate([this](const FCommandDesc& InCommand)
+	{
+		return Name == InCommand.Name;
+	});
+
+	check(Found);
+
+	return *Found;
+}
+
+}
\ No newline at end of file
diff --git a/Engine/Plugins/VirtualProduction/CaptureManager/CaptureManagerCore/Source/CaptureSourceFramework/Private/Widgets/SDefaultCaptureSourceCreator.cpp b/Engine/Plugins/VirtualProduction/CaptureManager/CaptureManagerCore/Source/CaptureSourceFramework/Private/Widgets/SDefaultCaptureSourceCreator.cpp
new file mode 100644
--- /dev/null
+++ b/Engine/Plugins/VirtualProduction/CaptureManager/CaptureManagerCore/Source/CaptureSourceFramework/Private/Widgets/SDefaultCaptureSourceCreator.cpp
@@ -0,0 +1,67 @@
+// Copyright Epic Games, Inc. All Rights Reserved.
+
+#include "Widgets/SDefaultCaptureSourceCreator.h"
+
+#include "WidgetUtils.h"
+
+#include "Widgets/SBoxPanel.h"
+#include "Widgets/Input/SButton.h"
+
+namespace UE::CaptureManager
+{
+
+void SDefaultCaptureSourceCreator::Construct(const FArguments& InArgs)
+{
+	FactoryId = InArgs._FactoryId;
+	CreationParams = InArgs._CreationParams;
+	OnCaptureSourceConfigured = InArgs._OnCaptureSourceConfigured;
+
+	TSharedPtr<SVerticalBox> VerticalBox = SNew(SVerticalBox);
+
+	Name.Set<FString>(TEXT("Unnamed"));
+
+	TSharedRef<SWidget> NameWidget =
+		UE::CaptureManager::Private::CreateParamWidget(FPropertyDesc(TEXT("Name"), FPropertyDesc::EType::String), Name).ToSharedRef();
+
+	VerticalBox->AddSlot()
+		.AutoHeight()
+		.VAlign(EVerticalAlignment::VAlign_Center)
+		.HAlign(EHorizontalAlignment::HAlign_Fill)
+		.AttachWidget(MoveTemp(NameWidget));
+
+	// Other creation params
+	for (const FPropertyDesc& Property : CreationParams)
+	{
+		CreationParamValues.Emplace(Property.Name, UE::CaptureManager::Private::GetDefaultValueForProperty(Property));
+
+		TSharedRef<SWidget> Widget =
+			UE::CaptureManager::Private::CreateParamWidget(Property, CreationParamValues[Property.Name]).ToSharedRef();
+
+		VerticalBox->AddSlot()
+			.AutoHeight()
+			.VAlign(EVerticalAlignment::VAlign_Center)
+			.HAlign(EHorizontalAlignment::HAlign_Fill)
+			.AttachWidget(MoveTemp(Widget));
+	}
+
+	VerticalBox->AddSlot()
+		.AutoHeight()
+		.VAlign(EVerticalAlignment::VAlign_Bottom)
+		.HAlign(EHorizontalAlignment::HAlign_Right)
+		[
+			SNew(SButton)
+				.Text(FText::FromString(TEXT("Create")))
+				.OnClicked_Lambda([this]() mutable
+			{
+				OnCaptureSourceConfigured.ExecuteIfBound(FactoryId, Name.Get<FString>(), CreationParamValues);
+				return FReply::Handled();
+			})
+		];
+
+	ChildSlot
+		[
+			VerticalBox.ToSharedRef()
+		];
+}
+
+}
\ No newline at end of file
diff --git a/Engine/Plugins/VirtualProduction/CaptureManager/CaptureManagerCore/Source/CaptureSourceFramework/Private/Widgets/SDefaultCaptureSourcePropertyCreator.cpp b/Engine/Plugins/VirtualProduction/CaptureManager/CaptureManagerCore/Source/CaptureSourceFramework/Private/Widgets/SDefaultCaptureSourcePropertyCreator.cpp
new file mode 100644
--- /dev/null
+++ b/Engine/Plugins/VirtualProduction/CaptureManager/CaptureManagerCore/Source/CaptureSourceFramework/Private/Widgets/SDefaultCaptureSourcePropertyCreator.cpp
@@ -0,0 +1,57 @@
+// Copyright Epic Games, Inc. All Rights Reserved.
+
+#include "Widgets/SDefaultCaptureSourcePropertyCreator.h"
+
+#include "WidgetUtils.h"
+
+#include "Widgets/SBoxPanel.h"
+
+namespace UE::CaptureManager
+{
+
+void SDefaultCaptureSourcePropertyCreator::Construct(const FArguments& InArgs)
+{
+	Capability = InArgs._Capability;
+	Name = InArgs._Name;
+	ValuePtr = InArgs._ValuePtr;
+
+	TSharedPtr<SHorizontalBox> HorizontalBox = SNew(SHorizontalBox);
+
+	FPropertyDesc PropertyDesc = FindPropertyDesc();
+
+	FPropertyValue& ValueRef = ValuePtr != nullptr ? *ValuePtr : Value;
+
+	ValueRef = Capability->GetPropertyValue(Name);
+
+	if (ValueRef.IsType<FEmptyVariantState>())
+	{
+		ValueRef = UE::CaptureManager::Private::GetDefaultValueForProperty(PropertyDesc);
+	}
+
+	HorizontalBox->AddSlot()
+		.AttachWidget(UE::CaptureManager::Private::MakeKeyWidget(PropertyDesc).ToSharedRef());
+
+	HorizontalBox->AddSlot()
+		.AttachWidget(UE::CaptureManager::Private::MakeValueWidget(Capability, PropertyDesc, ValueRef).ToSharedRef());
+
+	ChildSlot
+		[
+			HorizontalBox.ToSharedRef()
+		];
+}
+
+FPropertyDesc SDefaultCaptureSourcePropertyCreator::FindPropertyDesc()
+{
+	TArray<FPropertyDesc> Properties = Capability->GetProperties();
+
+	FPropertyDesc* Found = Properties.FindByPredicate([this](const FPropertyDesc& InProperty)
+	{
+		return Name == InProperty.Name;
+	});
+
+	check(Found);
+
+	return *Found;
+}
+
+}
\ No newline at end of file
diff --git a/Engine/Plugins/VirtualProduction/CaptureManager/CaptureManagerCore/Source/CaptureSourceFramework/Private/Widgets/WidgetUtils.cpp b/Engine/Plugins/VirtualProduction/CaptureManager/CaptureManagerCore/Source/CaptureSourceFramework/Private/Widgets/WidgetUtils.cpp
new file mode 100644
--- /dev/null
+++ b/Engine/Plugins/VirtualProduction/CaptureManager/CaptureManagerCore/Source/CaptureSourceFramework/Private/Widgets/WidgetUtils.cpp
@@ -0,0 +1,196 @@
+// Copyright Epic Games, Inc. All Rights Reserved.
+
+#include "WidgetUtils.h"
+
+#include "Widgets/Input/SButton.h"
+#include "Widgets/Input/SEditableTextBox.h"
+#include "Widgets/Input/SCheckBox.h"
+#include "Widgets/Input/SNumericEntryBox.h"
+
+namespace UE::CaptureManager::Private
+{
+
+TSharedPtr<SWidget> CreateParamWidget(const FPropertyDesc& InProperty, FPropertyValue& InPropertyValue)
+{
+	TSharedPtr<SHorizontalBox> HorizontalBox = SNew(SHorizontalBox);
+
+	HorizontalBox->AddSlot()
+		.AutoWidth()
+		.Padding(5.0f)
+		.AttachWidget(MakeKeyWidget(InProperty).ToSharedRef());
+
+	HorizontalBox->AddSlot()
+		.HAlign(EHorizontalAlignment::HAlign_Fill)
+		.AttachWidget(MakeValueWidget(InProperty, InPropertyValue).ToSharedRef());
+
+	return HorizontalBox;
+}
+
+TSharedPtr<SWidget> MakeKeyWidget(const FPropertyDesc& InProperty)
+{
+	TSharedPtr<STextBlock> TextBlock = SNew(STextBlock)
+		.Text(FText::FromString(InProperty.Name))
+		.TextStyle(FAppStyle::Get(), "NormalText");
+
+	return TextBlock;
+}
+
+TSharedPtr<SWidget> MakeValueWidget(const FPropertyDesc& InProperty, FPropertyValue& InPropertyValue)
+{
+	TSharedPtr<SWidget> Widget;
+	switch (InProperty.Type)
+	{
+		case FPropertyDesc::EType::Bool:
+			Widget = SNew(SCheckBox)
+				.OnCheckStateChanged_Lambda([&InPropertyValue](ECheckBoxState InCheckState)
+			{
+				InPropertyValue.Set<bool>(InCheckState == ECheckBoxState::Checked);
+			})
+				.IsChecked_Lambda([&InPropertyValue]()
+			{
+				return InPropertyValue.Get<bool>() ? ECheckBoxState::Checked : ECheckBoxState::Unchecked;
+			});
+			break;
+		case FPropertyDesc::EType::Number:
+			Widget = SNew(SNumericEntryBox<int64>)
+				.AllowSpin(false)
+				.OnValueCommitted_Lambda([&InPropertyValue](int64 InValue, ETextCommit::Type)
+			{
+				InPropertyValue.Set<int64>(InValue);
+			})
+				.Value_Lambda([&InPropertyValue]()
+			{
+				return InPropertyValue.Get<int64>();
+			});
+			break;
+		case FPropertyDesc::EType::FloatingPoint:
+			Widget = SNew(SNumericEntryBox<double>)
+				.AllowSpin(false)
+				.OnValueCommitted_Lambda([&InPropertyValue](double InValue, ETextCommit::Type)
+			{
+				InPropertyValue.Set<double>(InValue);
+			})
+				.Value_Lambda([&InPropertyValue]()
+			{
+				return InPropertyValue.Get<double>();
+			});
+			break;
+		case FPropertyDesc::EType::String:
+			Widget = SNew(SEditableTextBox)
+				.OnTextCommitted_Lambda([&InPropertyValue](FText InValue, ETextCommit::Type)
+			{
+				InPropertyValue.Set<FString>(InValue.ToString());
+			})
+				.Text_Lambda([&InPropertyValue]()
+			{
+				return FText::FromString(InPropertyValue.Get<FString>());
+			});
+			break;
+		case FPropertyDesc::EType::ArrayBool:
+		case FPropertyDesc::EType::ArrayNumber:
+		case FPropertyDesc::EType::ArrayString:
+		case FPropertyDesc::EType::ArrayFloatingPoint:
+		case FPropertyDesc::EType::Any:
+		default:
+			Widget = SNullWidget::NullWidget;
+	}
+
+	Widget->SetEnabled(InProperty.Access == FPropertyDesc::EAccess::ReadWrite);
+
+	return Widget;
+}
+
+TSharedPtr<SWidget> MakeValueWidget(FCaptureSourceCapability* InCapability,
+									const FPropertyDesc& InProperty,
+									FPropertyValue& InPropertyValue)
+{
+	TSharedPtr<SWidget> Widget;
+
+	switch (InProperty.Type)
+	{
+		case FPropertyDesc::EType::Bool:
+			Widget = SNew(SCheckBox)
+				.OnCheckStateChanged_Lambda([&InPropertyValue, InCapability, Name = InProperty.Name](ECheckBoxState InCheckState)
+			{
+				InPropertyValue.Set<bool>(InCheckState == ECheckBoxState::Checked);
+				InCapability->SetPropertyValue(Name, InPropertyValue);
+
+			})
+				.IsChecked_Lambda([&InPropertyValue]()
+			{
+				return InPropertyValue.Get<bool>() ? ECheckBoxState::Checked : ECheckBoxState::Unchecked;
+			});
+			break;
+		case FPropertyDesc::EType::Number:
+			Widget = SNew(SNumericEntryBox<int64>)
+				.OnValueCommitted_Lambda([&InPropertyValue, InCapability, Name = InProperty.Name](int64 InValue, ETextCommit::Type)
+			{
+				InPropertyValue.Set<int64>(InValue);
+				InCapability->SetPropertyValue(Name, InPropertyValue);
+			})
+				.Value_Lambda([&InPropertyValue]()
+			{
+				return InPropertyValue.Get<int64>();
+			});
+			break;
+		case FPropertyDesc::EType::FloatingPoint:
+			Widget = SNew(SNumericEntryBox<double>)
+				.OnValueCommitted_Lambda([&InPropertyValue, InCapability, Name = InProperty.Name](double InValue, ETextCommit::Type)
+			{
+				InPropertyValue.Set<double>(InValue);
+				InCapability->SetPropertyValue(Name, InPropertyValue);
+			})
+				.Value_Lambda([&InPropertyValue]()
+			{
+				return InPropertyValue.Get<double>();
+			});
+			break;
+		case FPropertyDesc::EType::String:
+			Widget = SNew(SEditableTextBox)
+				.OnTextCommitted_Lambda([&InPropertyValue, InCapability, Name = InProperty.Name](FText InValue, ETextCommit::Type)
+			{
+				InPropertyValue.Set<FString>(InValue.ToString());
+				InCapability->SetPropertyValue(Name, InPropertyValue);
+			})
+				.Text_Lambda([&InPropertyValue]()
+			{
+				return FText::FromString(InPropertyValue.Get<FString>());
+			});
+			break;
+		case FPropertyDesc::EType::ArrayBool:
+		case FPropertyDesc::EType::ArrayNumber:
+		case FPropertyDesc::EType::ArrayString:
+		case FPropertyDesc::EType::ArrayFloatingPoint:
+		case FPropertyDesc::EType::Any:
+		default:
+			Widget = SNullWidget::NullWidget;
+	}
+
+	Widget->SetEnabled(InProperty.Access == FPropertyDesc::EAccess::ReadWrite);
+
+	return Widget;
+}
+
+FPropertyValue GetDefaultValueForProperty(const FPropertyDesc& InProperty)
+{
+	switch (InProperty.Type)
+	{
+		case FPropertyDesc::EType::Bool:
+			return MakePropertyValue(false);
+		case FPropertyDesc::EType::Number:
+			return MakePropertyValue<int64>(0);
+		case FPropertyDesc::EType::FloatingPoint:
+			return MakePropertyValue<double>(0.0);
+		case FPropertyDesc::EType::String:
+			return MakePropertyValue<FString>(TEXT(""));
+		case FPropertyDesc::EType::ArrayBool:
+		case FPropertyDesc::EType::ArrayNumber:
+		case FPropertyDesc::EType::ArrayString:
+		case FPropertyDesc::EType::ArrayFloatingPoint:
+		case FPropertyDesc::EType::Any:
+		default:
+			return MakePropertyValue(FEmptyVariantState());
+	}
+}
+
+}
diff --git a/Engine/Plugins/VirtualProduction/CaptureManager/CaptureManagerCore/Source/CaptureSourceFramework/Private/Widgets/WidgetUtils.h b/Engine/Plugins/VirtualProduction/CaptureManager/CaptureManagerCore/Source/CaptureSourceFramework/Private/Widgets/WidgetUtils.h
new file mode 100644
--- /dev/null
+++ b/Engine/Plugins/VirtualProduction/CaptureManager/CaptureManagerCore/Source/CaptureSourceFramework/Private/Widgets/WidgetUtils.h
@@ -0,0 +1,18 @@
+// Copyright Epic Games, Inc. All Rights Reserved.
+
+#pragma once
+
+#include "CaptureSourceCapability.h"
+
+#include "Widgets/SWidget.h"
+
+namespace UE::CaptureManager::Private
+{
+
+TSharedPtr<SWidget> CreateParamWidget(const FPropertyDesc& InProperty, FPropertyValue& InPropertyValue);
+TSharedPtr<SWidget> MakeKeyWidget(const FPropertyDesc& InProperty);
+TSharedPtr<SWidget> MakeValueWidget(const FPropertyDesc& InProperty, FPropertyValue& InPropertyValue);
+TSharedPtr<SWidget> MakeValueWidget(FCaptureSourceCapability* InCapability, const FPropertyDesc& InProperty, FPropertyValue& InPropertyValue);
+FPropertyValue GetDefaultValueForProperty(const FPropertyDesc& InProperty);
+
+}
\ No newline at end of file
diff --git a/Engine/Plugins/VirtualProduction/CaptureManager/CaptureManagerCore/Source/CaptureSourceFramework/Public/Capabilities/IngestCapability.h b/Engine/Plugins/VirtualProduction/CaptureManager/CaptureManagerCore/Source/CaptureSourceFramework/Public/Capabilities/IngestCapability.h
new file mode 100644
--- /dev/null
+++ b/Engine/Plugins/VirtualProduction/CaptureManager/CaptureManagerCore/Source/CaptureSourceFramework/Public/Capabilities/IngestCapability.h
@@ -0,0 +1,228 @@
+// Copyright Epic Games, Inc. All Rights Reserved.
+
+#pragma once
+
+#include "CaptureSourceCapability.h"
+
+#include "Templates/ValueOrError.h"
+
+#include "CaptureManagerTakeMetadata.h"
+
+#include "Async/ManagedDelegate.h"
+
+namespace UE::CaptureManager
+{
+
+class CAPTURESOURCEFRAMEWORK_API FIngestCapabilityError
+{
+public:
+
+	enum ECode
+	{
+		Ok = 0,
+		AbortedByUser = 1,
+		InternalError = 2
+	};
+
+	FIngestCapabilityError(ECode InCode, FString InMessage);
+
+	ECode GetCode() const;
+	const FString& GetMessage() const;
+
+private:
+
+	ECode Code;
+	FString Message;
+};
+
+using FTakeId = int32;
+
+struct FTakeBaseEvent : public FCaptureEvent
+{
+	FTakeBaseEvent(FTakeId InTakeId, FString InName)
+		: FCaptureEvent(MoveTemp(InName))
+		, TakeId(InTakeId)
+	{
+	}
+
+	FTakeId TakeId;
+};
+
+struct CAPTURESOURCEFRAMEWORK_API FTakeAddedEvent : public FTakeBaseEvent
+{
+	static const FString Name;
+
+	explicit FTakeAddedEvent(FTakeId InTakeId)
+		: FTakeBaseEvent(InTakeId, Name)
+	{
+	}
+};
+
+struct CAPTURESOURCEFRAMEWORK_API FTakeUpdatedEvent : public FTakeBaseEvent
+{
+	static const FString Name;
+
+	explicit FTakeUpdatedEvent(FTakeId InTakeId)
+		: FTakeBaseEvent(InTakeId, Name)
+	{
+	}
+};
+
+struct CAPTURESOURCEFRAMEWORK_API FTakeRemovedEvent : public FTakeBaseEvent
+{
+	static const FString Name;
+
+	explicit FTakeRemovedEvent(FTakeId InTakeId)
+		: FTakeBaseEvent(InTakeId, Name)
+	{
+	}
+};
+
+// Corresponds to the MediaSample::EMediaTexturePixelFormat
+enum class EImagePixelFormat
+{
+	Undefined = -1,
+	U8_RGB = 0,
+	U8_BGR,
+	U8_RGBA,
+	U8_BGRA,
+	U8_I444,
+	U8_I420,
+	U8_NV12,
+	U8_Mono,
+	U16_Mono,
+	F_Mono,
+
+	Default = U8_RGB
+};
+
+struct FIngestOptions
+{
+	struct FVideo
+	{
+		FString FileNamePrefix;
+		FString Format;
+		EImagePixelFormat PixelFormat;
+	};
+
+	struct FAudio
+	{
+		FString FileNamePrefix;
+		FString Format;
+	};
+
+	FString Location;
+	FVideo Video;
+	FAudio Audio;
+};
+
+class CAPTURESOURCEFRAMEWORK_API FImportTakeCmd : public FCommandBase
+{
+public:
+	static const FString Name;
+
+	using FProcessProgressReporter = TManagedDelegate<FTakeId, double>;
+	using FProcessFinished = TManagedDelegate<FTakeId, TValueOrError<void, FIngestCapabilityError>>;
+
+	FImportTakeCmd(FIngestOptions InIngestOptions,
+				   FProcessProgressReporter InProgressReporter,
+				   FProcessFinished InProcessFinished);
+
+	FIngestOptions IngestOptions;
+	FProcessProgressReporter ProgressReporter;
+	FProcessFinished ProcessFinished;
+};
+
+class CAPTURESOURCEFRAMEWORK_API FUpdateTakeListCmd : public FCommandBase
+{
+public:
+
+	static const FString Name;
+
+	using FCallback = TManagedDelegate<TArray<FTakeId>>;
+
+	FUpdateTakeListCmd(FCallback InCallback);
+
+	FCallback Callback;
+};
+
+class CAPTURESOURCEFRAMEWORK_API FIngestCapabilityBase : public FCaptureSourceCapability
+{
+public:
+
+	class CAPTURESOURCEFRAMEWORK_API FTakeProcessHandle
+	{
+	public:
+
+		FTakeProcessHandle() = default;
+		FTakeProcessHandle(FIngestCapabilityBase* InOwner, FTakeId InId);
+
+		void Cancel();
+
+	private:
+
+		FIngestCapabilityBase* Owner;
+		FTakeId Id;
+	};
+
+	static const FString Name;
+
+	static const FString CancelCmd;
+
+	static const FString Id;
+
+	FIngestCapabilityBase();
+
+	virtual ~FIngestCapabilityBase() = default;
+
+	virtual FTakeProcessHandle ImportTakeAsync(FTakeId InTakeId,
+											   FIngestOptions InIngestOptions,
+											   FImportTakeCmd::FProcessFinished InProcessFinished,
+											   FImportTakeCmd::FProcessProgressReporter InProcessProgressReported) = 0;
+
+	virtual void UpdateTakeListAsync(FUpdateTakeListCmd::FCallback InUpdateTakeListCallback) = 0;
+
+	FTakeId AddTakeMetadata(FTakeMetadata InTakeInfo);
+	void RemoveTakeMetadata(FTakeId InTakeId);
+	void UpdateTakeMetadata(FTakeId InTakeId, FTakeMetadata InTakeInfo);
+	TOptional<FTakeMetadata> GetTakeMetadata(FTakeId InTakeId) const;
+	TArray<FTakeId> GetTakeIds() const;
+
+	void PublishTakeAddedEvent(FTakeId InTakeId);
+	void PublishTakeUpdatedEvent(FTakeId InTakeId);
+	void PublishTakeRemovedEvent(FTakeId InTakeId);
+
+protected:
+
+	virtual void Cancel(FTakeId InId) = 0;
+
+	mutable FCriticalSection Mutex;
+	TMap<FTakeId, FTakeMetadata> Takes;
+
+private:
+
+	std::atomic<FTakeId> CurrentTakeId;
+};
+
+class CAPTURESOURCEFRAMEWORK_API FIngestCapability : public FIngestCapabilityBase
+{
+public:
+
+	FIngestCapability();
+
+	FTakeProcessHandle ImportTakeAsync(FTakeId InTakeId,
+									   FIngestOptions InIngestOptions,
+									   FImportTakeCmd::FProcessFinished InProcessFinished,
+									   FImportTakeCmd::FProcessProgressReporter InProcessProgressReported) override;
+	void UpdateTakeListAsync(FUpdateTakeListCmd::FCallback InUpdateTakeListCallback) override;
+
+	void SetImportTakeHandler(FCommandHandler InCommandHandler);
+	void SetUpdateTakeListHandler(FCommandHandler InCommandHandler);
+	void SetCancelHandler(FCommandHandler InCommandHandler);
+
+private:
+
+	void Cancel(FTakeId InTakeId);
+};
+
+}
\ No newline at end of file
diff --git a/Engine/Plugins/VirtualProduction/CaptureManager/CaptureManagerCore/Source/CaptureSourceFramework/Public/CaptureSourceDetailsCustomization.h b/Engine/Plugins/VirtualProduction/CaptureManager/CaptureManagerCore/Source/CaptureSourceFramework/Public/CaptureSourceDetailsCustomization.h
new file mode 100644
--- /dev/null
+++ b/Engine/Plugins/VirtualProduction/CaptureManager/CaptureManagerCore/Source/CaptureSourceFramework/Public/CaptureSourceDetailsCustomization.h
@@ -0,0 +1,78 @@
+// Copyright Epic Games, Inc. All Rights Reserved.
+
+#pragma once
+
+#include "CaptureSourceCapability.h"
+#include "CaptureSourceFactory.h"
+
+#include "Widgets/SWidget.h"
+
+#include "Widgets/SDefaultCaptureSourceCreator.h"
+
+namespace UE::CaptureManager
+{
+
+class CAPTURESOURCEFRAMEWORK_API FCaptureSourceDetailsCustomization
+{
+public:
+
+	DECLARE_DELEGATE_RetVal_ThreeParams(TSharedPtr<SWidget>, FPropertyWidgetCreator, const FString&, FCaptureSourceCapability*, const FPropertyDesc&);
+	DECLARE_DELEGATE_RetVal_ThreeParams(TSharedPtr<SWidget>, FCommandWidgetCreator, const FString&, FCaptureSourceCapability*, const FCommandDesc&);
+	DECLARE_DELEGATE_RetVal_TwoParams(TSharedPtr<SWidget>, FCapabilityWidgetCreator, const FString&, FCaptureSourceCapability*);
+	DECLARE_DELEGATE_RetVal_TwoParams(TSharedPtr<SWidget>, FCaptureSourceWidgetCreator, const FCaptureSourceDescriptor&, FOnCaptureSourceConfigured);
+
+	FCaptureSourceDetailsCustomization();
+
+	void RegisterPropertyCustomization(const FString& InCaptureSourceClassId,
+									   const FString& InCapabilityName,
+									   const FString& InPropertyName,
+									   FPropertyWidgetCreator InWidgetCreator);
+
+	void RegisterCommandCustomization(const FString& InCaptureSourceClassId,
+									  const FString& InCapabilityName,
+									  const FString& InCommandName,
+									  FCommandWidgetCreator InWidgetCreator);
+
+	void RegisterCapabilityCustomization(const FString& InCaptureSourceClassId,
+										 const FString& InCapabilityName,
+										 FCapabilityWidgetCreator InWidgetCreator);
+
+	void RegisterCaptureSourceCustomization(const FString& InCaptureSourceClassId,
+											FCaptureSourceWidgetCreator InWidgetCreator);
+
+	TSharedPtr<SWidget> CreatePropertyWidget(const FString& InCaptureSourceClassId,
+											 FCaptureSourceCapability* InCapability,
+											 const FPropertyDesc& InProperty);
+
+	TSharedPtr<SWidget> CreateCommandWidget(const FString& InCaptureSourceClassId,
+											FCaptureSourceCapability* InCapability,
+											const FCommandDesc& InCommand);
+
+	TSharedPtr<SWidget> CreateCapabilityWidget(const FString& InCaptureSourceClassId,
+											   FCaptureSourceCapability* InCapability);
+
+	TSharedPtr<SWidget> CreateCaptureSourceWidget(const FCaptureSourceDescriptor& InCaptureSourceDescriptor,
+												  FOnCaptureSourceConfigured InCaptureSourceCreated);
+
+	bool ContainsPropertyWidgetCustomization(const FString& InFactoryId,
+											 const FString& InCapabilityName,
+											 const FString& InPropertyName) const;
+
+	bool ContainsCommandWidgetCustomization(const FString& InFactoryId,
+											const FString& InCapabilityName,
+											const FString& InCommandName) const;
+
+	bool ContainsCapabilityWidgetCustomization(const FString& InFactoryId,
+											   const FString& InCapabilityName) const;
+
+	bool ContainsCaptureSourceWidgetCustomization(const FString& InFactoryId) const;
+
+private:
+
+	TMap<FString, FPropertyWidgetCreator> PropertyCustomizations;
+	TMap<FString, FCommandWidgetCreator> CommandCustomizations;
+	TMap<FString, FCaptureSourceWidgetCreator> CaptureSourceCustomizations;
+	TMap<FString, FCapabilityWidgetCreator> CapabilityCustomizations;
+};
+
+}
\ No newline at end of file
diff --git a/Engine/Plugins/VirtualProduction/CaptureManager/CaptureManagerCore/Source/CaptureSourceFramework/Public/CaptureSourceFactory.h b/Engine/Plugins/VirtualProduction/CaptureManager/CaptureManagerCore/Source/CaptureSourceFramework/Public/CaptureSourceFactory.h
new file mode 100644
--- /dev/null
+++ b/Engine/Plugins/VirtualProduction/CaptureManager/CaptureManagerCore/Source/CaptureSourceFramework/Public/CaptureSourceFactory.h
@@ -0,0 +1,65 @@
+// Copyright Epic Games, Inc. All Rights Reserved.
+
+#pragma once
+
+#include "CaptureSource.h"
+#include "CaptureSourceError.h" 
+
+class CAPTURESOURCEFRAMEWORK_API FCaptureSourceDescriptor
+{
+public:
+
+	FCaptureSourceDescriptor(const FString& InId, FText InSourceTypeDisplayName, TArray<FPropertyDesc> InCreationParams);
+
+	void SetDiscoverable(bool bInIsDiscoverable);
+	bool IsDiscoverable() const;
+
+	void AddCreationParam(FPropertyDesc InParam);
+
+	template<typename ... FArgs>
+	void EmplaceCreationParam(FArgs&&... InArgs)
+	{
+		CreationParams.Emplace(Forward<FArgs>(InArgs)...);
+	}
+
+	const TArray<FPropertyDesc>& GetCreationParams() const;
+
+	const FString& GetId() const;
+
+	// The display name for the source TYPE, for example "LiveLink Face Archive"
+	const FText& GetSourceTypeDisplayName() const;
+
+private:
+
+	FString Id;
+	FText SourceTypeDisplayName;
+	TArray<FPropertyDesc> CreationParams;
+	bool bIsDiscoverable = true;
+};
+
+class CAPTURESOURCEFRAMEWORK_API FCaptureSourceFactory
+{
+public:
+
+	using FCaptureSourceResult = FCaptureValueResult<TUniquePtr<UE::CaptureManager::FCaptureSource>>;
+
+	DECLARE_DELEGATE_OneParam(FDiscoverCallback, TArray<FCaptureSourceResult>);
+
+	virtual ~FCaptureSourceFactory() = default;
+
+	virtual FText GetCaptureSourceFactoryName() const = 0;
+	virtual FString GetCaptureSourceFactoryId() const = 0;
+
+	virtual bool IsDiscoverable() const = 0;
+	FCaptureVoidResult Discover(FDiscoverCallback InCallback);
+
+	FCaptureSourceResult CreateCaptureSource(const FString& InName, TMap<FString, FPropertyValue> InCreationParamsValue);
+	FCaptureSourceDescriptor CreateCaptureSourceDescriptor() const;
+
+private:
+
+	virtual void DiscoverImpl(FDiscoverCallback InCallback, FCaptureVoidResult& OutResult) {}
+
+	virtual FCaptureSourceResult CreateCaptureSourceImpl(const FString& InName, const TMap<FString, FPropertyValue>& InCreationParamsValue) = 0;
+	virtual void CreateCaptureSourceDescriptorImpl(FCaptureSourceDescriptor& OutDescriptor) const = 0;
+};
\ No newline at end of file
diff --git a/Engine/Plugins/VirtualProduction/CaptureManager/CaptureManagerCore/Source/CaptureSourceFramework/Public/Widgets/SDefaultCaptureSourceCapabilityCreator.h b/Engine/Plugins/VirtualProduction/CaptureManager/CaptureManagerCore/Source/CaptureSourceFramework/Public/Widgets/SDefaultCaptureSourceCapabilityCreator.h
new file mode 100644
--- /dev/null
+++ b/Engine/Plugins/VirtualProduction/CaptureManager/CaptureManagerCore/Source/CaptureSourceFramework/Public/Widgets/SDefaultCaptureSourceCapabilityCreator.h
@@ -0,0 +1,36 @@
+// Copyright Epic Games, Inc. All Rights Reserved.
+
+#pragma once
+
+#include "CaptureSourceCapability.h"
+
+#include "Widgets/SCompoundWidget.h"
+
+namespace UE::CaptureManager
+{
+
+class CAPTURESOURCEFRAMEWORK_API SDefaultCaptureSourceCapabilityCreator : public SCompoundWidget
+{
+public:
+
+	SLATE_BEGIN_ARGS(SDefaultCaptureSourceCapabilityCreator)
+		{
+		}
+		SLATE_ARGUMENT(FString, FactoryId)
+		SLATE_ARGUMENT(FCaptureSourceCapability*, Capability)
+
+	SLATE_END_ARGS()
+
+	void Construct(const FArguments& InArgs);
+
+private:
+
+	void HandlePropertyChangedEvent(TSharedPtr<const FCaptureEvent> InEvent);
+
+	FString FactoryId;
+
+	FCaptureSourceCapability* Capability;
+	TMap<FString, FPropertyValue> PropertyValues;
+};
+
+}
\ No newline at end of file
diff --git a/Engine/Plugins/VirtualProduction/CaptureManager/CaptureManagerCore/Source/CaptureSourceFramework/Public/Widgets/SDefaultCaptureSourceCommandCreator.h b/Engine/Plugins/VirtualProduction/CaptureManager/CaptureManagerCore/Source/CaptureSourceFramework/Public/Widgets/SDefaultCaptureSourceCommandCreator.h
new file mode 100644
--- /dev/null
+++ b/Engine/Plugins/VirtualProduction/CaptureManager/CaptureManagerCore/Source/CaptureSourceFramework/Public/Widgets/SDefaultCaptureSourceCommandCreator.h
@@ -0,0 +1,36 @@
+// Copyright Epic Games, Inc. All Rights Reserved.
+
+#pragma once
+
+#include "CaptureSourceCapability.h"
+
+#include "Widgets/SCompoundWidget.h"
+
+namespace UE::CaptureManager
+{
+
+class CAPTURESOURCEFRAMEWORK_API SDefaultCaptureSourceCommandCreator : public SCompoundWidget
+{
+public:
+
+	SLATE_BEGIN_ARGS(SDefaultCaptureSourceCommandCreator)
+		{
+		}
+		SLATE_ARGUMENT(FCaptureSourceCapability*, Capability)
+		SLATE_ARGUMENT(FString, Name)
+
+	SLATE_END_ARGS()
+
+	void Construct(const FArguments& InArgs);
+
+private:
+
+	FCommandDesc FindCommandDesc();
+
+	FCaptureSourceCapability* Capability;
+
+	FString Name;
+	FCommandBase::FValues Params;
+};
+
+}
\ No newline at end of file
diff --git a/Engine/Plugins/VirtualProduction/CaptureManager/CaptureManagerCore/Source/CaptureSourceFramework/Public/Widgets/SDefaultCaptureSourceCreator.h b/Engine/Plugins/VirtualProduction/CaptureManager/CaptureManagerCore/Source/CaptureSourceFramework/Public/Widgets/SDefaultCaptureSourceCreator.h
new file mode 100644
--- /dev/null
+++ b/Engine/Plugins/VirtualProduction/CaptureManager/CaptureManagerCore/Source/CaptureSourceFramework/Public/Widgets/SDefaultCaptureSourceCreator.h
@@ -0,0 +1,39 @@
+// Copyright Epic Games, Inc. All Rights Reserved.
+
+#pragma once
+
+#include "CaptureSourceCapabilityProperty.h"
+
+#include "Widgets/SCompoundWidget.h"
+
+namespace UE::CaptureManager
+{
+
+using FOnCaptureSourceConfigured = TDelegate<void(FString, FString, TMap<FString, FPropertyValue>)>;
+
+class CAPTURESOURCEFRAMEWORK_API SDefaultCaptureSourceCreator : public SCompoundWidget
+{
+public:
+
+	SLATE_BEGIN_ARGS(SDefaultCaptureSourceCreator)
+		{
+		}
+
+		SLATE_ARGUMENT(FString, FactoryId)
+		SLATE_ARGUMENT(TArray<FPropertyDesc>, CreationParams)
+		SLATE_EVENT(FOnCaptureSourceConfigured, OnCaptureSourceConfigured)
+	SLATE_END_ARGS()
+
+	void Construct(const FArguments& InArgs);
+
+private:
+
+	FString FactoryId;
+	TArray<FPropertyDesc> CreationParams;
+	FOnCaptureSourceConfigured OnCaptureSourceConfigured;
+
+	FPropertyValue Name;
+	TMap<FString, FPropertyValue> CreationParamValues;
+};
+
+}
\ No newline at end of file
diff --git a/Engine/Plugins/VirtualProduction/CaptureManager/CaptureManagerCore/Source/CaptureSourceFramework/Public/Widgets/SDefaultCaptureSourcePropertyCreator.h b/Engine/Plugins/VirtualProduction/CaptureManager/CaptureManagerCore/Source/CaptureSourceFramework/Public/Widgets/SDefaultCaptureSourcePropertyCreator.h
new file mode 100644
--- /dev/null
+++ b/Engine/Plugins/VirtualProduction/CaptureManager/CaptureManagerCore/Source/CaptureSourceFramework/Public/Widgets/SDefaultCaptureSourcePropertyCreator.h
@@ -0,0 +1,38 @@
+// Copyright Epic Games, Inc. All Rights Reserved.
+
+#pragma once
+
+#include "CaptureSourceCapability.h"
+
+#include "Widgets/SCompoundWidget.h"
+
+namespace UE::CaptureManager
+{
+
+class CAPTURESOURCEFRAMEWORK_API SDefaultCaptureSourcePropertyCreator : public SCompoundWidget
+{
+public:
+
+	SLATE_BEGIN_ARGS(SDefaultCaptureSourcePropertyCreator)
+		{
+		}
+		SLATE_ARGUMENT(FCaptureSourceCapability*, Capability)
+		SLATE_ARGUMENT(FString, Name)
+		SLATE_ARGUMENT(FPropertyValue*, ValuePtr)
+
+	SLATE_END_ARGS()
+
+	void Construct(const FArguments& InArgs);
+
+private:
+
+	FPropertyDesc FindPropertyDesc();
+
+	FCaptureSourceCapability* Capability;
+
+	FString Name;
+	FPropertyValue* ValuePtr = nullptr;
+	FPropertyValue Value;
+};
+
+}
\ No newline at end of file
diff --git a/Engine/Plugins/VirtualProduction/CaptureManager/CaptureManagerCore/Source/CaptureUtils/Public/Async/Task.h b/Engine/Plugins/VirtualProduction/CaptureManager/CaptureManagerCore/Source/CaptureUtils/Public/Async/Task.h
new file mode 100644
--- /dev/null
+++ b/Engine/Plugins/VirtualProduction/CaptureManager/CaptureManagerCore/Source/CaptureUtils/Public/Async/Task.h
@@ -0,0 +1,37 @@
+// Copyright Epic Games, Inc. All Rights Reserved.
+
+#pragma once
+
+#include "Templates/Function.h"
+#include "Async/StopToken.h"
+
+template<typename TTask>
+class FAsyncTask;
+
+namespace UE::CaptureManager
+{
+
+class CAPTUREUTILS_API FCancelableAsyncTask
+{
+public:
+	using FTaskFunction = TFunction<void(const FStopToken&)>;
+
+	explicit FCancelableAsyncTask(FTaskFunction InTaskFunction);
+
+	~FCancelableAsyncTask();
+
+	bool IsDone();
+
+	void StartSync();
+	void StartAsync();
+
+	void Cancel();
+
+private:
+	class FAsyncTaskInternal;
+
+	FStopRequester StopRequester;
+	TUniquePtr<FAsyncTask<FAsyncTaskInternal>> AsyncTask;
+};
+
+}
diff --git a/Engine/Plugins/VirtualProduction/CaptureManager/CaptureManagerCore/Source/CaptureUtils/Public/Async/TaskProgress.h b/Engine/Plugins/VirtualProduction/CaptureManager/CaptureManagerCore/Source/CaptureUtils/Public/Async/TaskProgress.h
new file mode 100644
--- /dev/null
+++ b/Engine/Plugins/VirtualProduction/CaptureManager/CaptureManagerCore/Source/CaptureUtils/Public/Async/TaskProgress.h
@@ -0,0 +1,54 @@
+// Copyright Epic Games, Inc. All Rights Reserved.
+
+#pragma once
+
+#include "Containers/Array.h"
+#include "HAL/CriticalSection.h"
+#include "Delegates/Delegate.h"
+
+#include "Monitor.h"
+
+namespace UE::CaptureManager
+{
+
+class CAPTUREUTILS_API FTaskProgress
+{
+public:
+
+	DECLARE_DELEGATE_OneParam(FProgressReporter, double InProgress);
+
+	class CAPTUREUTILS_API FTask
+	{
+	public:
+
+		void Update(double InProgress);
+
+	private:
+
+		FTask(FTaskProgress& InTaskProgress, int32 InId);
+
+		FTaskProgress& TaskProgress;
+		int32 Id;
+
+		friend FTaskProgress;
+	};
+
+	FTaskProgress(uint32 InAmountOfWork, FProgressReporter InReport);
+
+	FTask StartTask();
+
+	double GetTotalProgress() const;
+
+private:
+
+	void Update(int32 InTaskToUpdate, double InProgress);
+	void Report() const;
+
+	FProgressReporter Reporter;
+
+	std::atomic<int32> CurrentTask;
+
+	mutable FMonitor<TArray<double>> CurrentProgressValues;
+};
+
+}
\ No newline at end of file
diff --git a/Engine/Plugins/VirtualProduction/CaptureManager/CaptureManagerCore/Source/CaptureUtils/Public/CaptureUtilsModule.h b/Engine/Plugins/VirtualProduction/CaptureManager/CaptureManagerCore/Source/CaptureUtils/Public/CaptureUtilsModule.h
new file mode 100644
--- /dev/null
+++ b/Engine/Plugins/VirtualProduction/CaptureManager/CaptureManagerCore/Source/CaptureUtils/Public/CaptureUtilsModule.h
@@ -0,0 +1,21 @@
+// Copyright Epic Games, Inc. All Rights Reserved.
+
+#pragma once
+
+#include "Modules/ModuleManager.h"
+
+#include "Async/CaptureTimerManager.h"
+
+class CAPTUREUTILS_API FCaptureUtilsModule : public IModuleInterface
+{
+public:
+
+	virtual void StartupModule() override;
+	virtual void ShutdownModule() override;
+
+	TSharedRef<UE::CaptureManager::FCaptureTimerManager> GetTimerManager();
+
+private:
+
+	TSharedPtr<UE::CaptureManager::FCaptureTimerManager> TimerManager;
+};
\ No newline at end of file
diff --git a/Engine/Plugins/VirtualProduction/CaptureManager/CaptureManagerCore/Source/DataIngestCore/Public/Utils/UnrealCalibrationParser.h b/Engine/Plugins/VirtualProduction/CaptureManager/CaptureManagerCore/Source/DataIngestCore/Public/Utils/UnrealCalibrationParser.h
new file mode 100644
--- /dev/null
+++ b/Engine/Plugins/VirtualProduction/CaptureManager/CaptureManagerCore/Source/DataIngestCore/Public/Utils/UnrealCalibrationParser.h
@@ -0,0 +1,14 @@
+// Copyright Epic Games, Inc. All Rights Reserved.
+
+#pragma once
+
+#include "CameraCalibration.h"
+
+class DATAINGESTCORE_API FUnrealCalibrationParser
+{
+public:
+	using FParseResult = TValueOrError<TArray<FCameraCalibration>, FText>;
+
+	/** Parse input file into unreal calibration format. */
+	static FParseResult Parse(const FString& InFile);
+};
\ No newline at end of file
diff --git a/Engine/Plugins/VirtualProduction/CaptureManager/CaptureManagerCore/Source/LiveLinkHubCaptureMessaging/Public/Features/ConnectAcceptor.h b/Engine/Plugins/VirtualProduction/CaptureManager/CaptureManagerCore/Source/LiveLinkHubCaptureMessaging/Public/Features/ConnectAcceptor.h
new file mode 100644
--- /dev/null
+++ b/Engine/Plugins/VirtualProduction/CaptureManager/CaptureManagerCore/Source/LiveLinkHubCaptureMessaging/Public/Features/ConnectAcceptor.h
@@ -0,0 +1,51 @@
+// Copyright Epic Games, Inc. All Rights Reserved.
+
+#pragma once
+
+#include "MessageEndpointBuilder.h"
+#include "MessageEndpoint.h"
+
+#include "LiveLinkHubCaptureMessages.h"
+
+#include "Messenger.h"
+
+#include "Async/CaptureTimerManager.h"
+
+class LIVELINKHUBCAPTUREMESSAGING_API FConnectAcceptor : public FFeatureBase
+{
+public:
+
+	DECLARE_DELEGATE_RetVal_TwoParams(FConnectResponse*, FConnectAccepted, const FConnectRequest& InRequest, const FMessageAddress& InAddress);
+	DECLARE_DELEGATE_OneParam(FConnectionLostHandler, const FMessageAddress& InAddress);
+
+	FConnectAcceptor();
+	virtual ~FConnectAcceptor() override;
+
+	void SetConnectionHandler(FConnectAccepted InConnectHandler,
+							  FConnectionLostHandler InConnectionLostHandler);
+
+protected:
+
+	void Initialize(FMessageEndpointBuilder& InBuilder);
+
+private:
+
+	void ConnectRequestHandler(const FConnectRequest& InRequest, const TSharedRef<IMessageContext, ESPMode::ThreadSafe>& InContext);
+	void HandleKeepAlive(const FPingMessage& InRequest, const TSharedRef<IMessageContext, ESPMode::ThreadSafe>& InContext);
+
+	void CheckConnectionActivity();
+
+	double LastPingRequest = 0.0f;
+	double InactivityTimeout = 15.0f; // seconds
+
+	FConnectAccepted ConnectHandler;
+	FConnectionLostHandler ConnectionLostHandler;
+
+	using FFeatureBase::Endpoint;
+	using FFeatureBase::Address;
+
+	TSharedRef<UE::CaptureManager::FCaptureTimerManager> TimerManager;
+	UE::CaptureManager::FCaptureTimerManager::FTimerHandle ClientActivity;
+
+	std::atomic_bool bConnected = false;
+};
\ No newline at end of file
diff --git a/Engine/Plugins/VirtualProduction/CaptureManager/CaptureManagerCore/Source/LiveLinkHubCaptureMessaging/Public/Features/ConnectStarter.h b/Engine/Plugins/VirtualProduction/CaptureManager/CaptureManagerCore/Source/LiveLinkHubCaptureMessaging/Public/Features/ConnectStarter.h
new file mode 100644
--- /dev/null
+++ b/Engine/Plugins/VirtualProduction/CaptureManager/CaptureManagerCore/Source/LiveLinkHubCaptureMessaging/Public/Features/ConnectStarter.h
@@ -0,0 +1,63 @@
+// Copyright Epic Games, Inc. All Rights Reserved.
+
+#pragma once
+
+#include "MessageEndpointBuilder.h"
+#include "MessageEndpoint.h"
+
+#include "Messenger.h"
+
+#include "Async/CaptureTimerManager.h"
+
+#include "LiveLinkHubCaptureMessages.h"
+
+class LIVELINKHUBCAPTUREMESSAGING_API FConnectStarter : 
+	public FFeatureBase, 
+	public TSharedFromThis<FConnectStarter>
+{
+public:
+
+	DECLARE_DELEGATE_OneParam(FConnectHandler, const FConnectResponse& InResponse);
+	DECLARE_DELEGATE(FDisconnectHandler);
+
+	FConnectStarter();
+	virtual ~FConnectStarter() override;
+
+	void Connect(FConnectHandler InConnectHandler);
+	void SetDisconnectHandler(FDisconnectHandler InHandler);
+
+	bool IsConnected() const;
+
+protected:
+
+	void Initialize(FMessageEndpointBuilder& InBuilder);
+
+private:
+
+	static const float KeepAliveInterval;
+	static const float KeepAliveTimeout;
+
+	DECLARE_DELEGATE_OneParam(FPingCallback, const FPongMessage& InResponse);
+	void SendPingRequest(FPingCallback InCallback);
+
+	void HandleConnectResponse(const FConnectResponse& InResponse, const TSharedRef<IMessageContext, ESPMode::ThreadSafe>& InContext);
+	void HandlePingResponse(const FPongMessage& InResponse, const TSharedRef<IMessageContext, ESPMode::ThreadSafe>& InContext);
+	void AddContext(FGuid InGuid, FConnectHandler InHandler);
+
+	void OnKeepAliveInterval();
+
+	using FFeatureBase::Endpoint;
+	using FFeatureBase::Address;
+
+	FCriticalSection Mutex;
+	TMap<FGuid, FConnectHandler> Contexts;
+
+	FCriticalSection PingMutex;
+	TMap<FGuid, FPingCallback> PingContexts;
+
+	TSharedRef<UE::CaptureManager::FCaptureTimerManager> TimerManager;
+	UE::CaptureManager::FCaptureTimerManager::FTimerHandle KeepAlive;
+	FDisconnectHandler Handler;
+
+	std::atomic_bool bConnected = false;
+};
\ No newline at end of file
diff --git a/Engine/Plugins/VirtualProduction/CaptureManager/CaptureManagerCore/Source/LiveLinkHubCaptureMessaging/Public/Features/UploadStateHandler.h b/Engine/Plugins/VirtualProduction/CaptureManager/CaptureManagerCore/Source/LiveLinkHubCaptureMessaging/Public/Features/UploadStateHandler.h
new file mode 100644
--- /dev/null
+++ b/Engine/Plugins/VirtualProduction/CaptureManager/CaptureManagerCore/Source/LiveLinkHubCaptureMessaging/Public/Features/UploadStateHandler.h
@@ -0,0 +1,38 @@
+// Copyright Epic Games, Inc. All Rights Reserved.
+
+#pragma once
+
+#include "MessageEndpointBuilder.h"
+#include "MessageEndpoint.h"
+
+#include "Messenger.h"
+
+#include "LiveLinkHubCaptureMessages.h"
+
+class LIVELINKHUBCAPTUREMESSAGING_API FUploadStateHandler : public FFeatureBase
+{
+public:
+
+	DECLARE_DELEGATE_ThreeParams(FUploadStateCallback, const FGuid& InCaptureSourceId, const FGuid& InTakeUploadId, double InProgress);
+	DECLARE_DELEGATE_FourParams(FUploadFinishedCallback, const FGuid& InCaptureSourceId, const FGuid& InTakeUploadId, FString InMessage, int32 InCode);
+
+	FUploadStateHandler();
+
+	void SetUploadCallbacks(FUploadStateCallback InStateCallback,
+							FUploadFinishedCallback InFinishedCallback);
+
+protected:
+
+	void Initialize(FMessageEndpointBuilder& InBuilder);
+
+private:
+
+	using FFeatureBase::Endpoint;
+	using FFeatureBase::Address;
+
+	void HandleUploadStateMessage(const FUploadState& InMessage, const TSharedRef<IMessageContext, ESPMode::ThreadSafe>& InContext);
+	void HandleUploadFinishedMessage(const FUploadFinished& InMessage, const TSharedRef<IMessageContext, ESPMode::ThreadSafe>& InContext);
+
+	FUploadStateCallback StateCallback;
+	FUploadFinishedCallback FinishedCallback;
+};
\ No newline at end of file
diff --git a/Engine/Plugins/VirtualProduction/CaptureManager/CaptureManagerCore/Source/LiveLinkHubCaptureMessaging/Public/Features/UploadStateSender.h b/Engine/Plugins/VirtualProduction/CaptureManager/CaptureManagerCore/Source/LiveLinkHubCaptureMessaging/Public/Features/UploadStateSender.h
new file mode 100644
--- /dev/null
+++ b/Engine/Plugins/VirtualProduction/CaptureManager/CaptureManagerCore/Source/LiveLinkHubCaptureMessaging/Public/Features/UploadStateSender.h
@@ -0,0 +1,31 @@
+// Copyright Epic Games, Inc. All Rights Reserved.
+
+#pragma once
+
+#include "MessageEndpointBuilder.h"
+#include "MessageEndpoint.h"
+
+#include "Messenger.h"
+
+#include "LiveLinkHubCaptureMessages.h"
+
+class LIVELINKHUBCAPTUREMESSAGING_API FUploadStateSender : public FFeatureBase
+{
+public:
+
+	FUploadStateSender();
+
+	void SendUploadStateMessage(const FGuid& InCaptureSourceId, const FGuid& InTakeUploadId, double InProgress);
+	void SendUploadDoneMessage(const FGuid& InCaptureSourceId, const FGuid& InTakeUploadId, FString InMessage, int32 InCode);
+
+protected:
+
+	void Initialize(FMessageEndpointBuilder& InBuilder);
+
+private:
+
+	EStatus ConvertStatus(int32 InCode);
+
+	using FFeatureBase::Endpoint;
+	using FFeatureBase::Address;
+};
\ No newline at end of file
diff --git a/Engine/Plugins/VirtualProduction/CaptureManager/CaptureManagerCore/Source/LiveLinkHubCaptureMessaging/Public/UploadDataMessage.h b/Engine/Plugins/VirtualProduction/CaptureManager/CaptureManagerCore/Source/LiveLinkHubCaptureMessaging/Public/UploadDataMessage.h
new file mode 100644
--- /dev/null
+++ b/Engine/Plugins/VirtualProduction/CaptureManager/CaptureManagerCore/Source/LiveLinkHubCaptureMessaging/Public/UploadDataMessage.h
@@ -0,0 +1,81 @@
+// Copyright Epic Games, Inc. All Rights Reserved.
+
+#pragma once
+
+#include "Containers/UnrealString.h"
+#include "Containers/StaticArray.h"
+
+#include "Internationalization/Text.h"
+
+#include "Templates/ValueOrError.h"
+
+#include "Network/TcpReaderWriter.h"
+
+#include "Misc/Guid.h"
+
+class LIVELINKHUBCAPTUREMESSAGING_API FUploadError
+{
+public:
+
+	FUploadError(FText InMessage, int32 InCode = 0);
+
+	const FText& GetText() const;
+	int32 GetCode() const;
+
+private:
+
+	FText Message;
+	int32 Code;
+};
+
+template<typename T>
+using FUploadResult = TValueOrError<T, FUploadError>;
+
+using FUploadVoidResult = FUploadResult<void>;
+
+struct FUploadDataHeader
+{
+	static const TArray<uint8> Header;
+
+	FGuid ClientId;
+	FGuid CaptureSourceId;
+	FGuid TakeUploadId;
+
+	FString CaptureSourceName;
+	FString TakeName;
+	uint64 TotalLength = 0;
+};
+
+struct FUploadFileDataHeader
+{
+	FString FileName;
+	uint64 Length = 0;
+};
+
+class LIVELINKHUBCAPTUREMESSAGING_API FUploadDataMessage
+{
+public:
+
+	static constexpr int32 HashSize = 16;
+
+	static FUploadVoidResult SerializeHeader(FUploadDataHeader InHeader, UE::CaptureManager::ITcpSocketWriter& InWriter);
+	static FUploadVoidResult SerializeFileHeader(FUploadFileDataHeader InFileHeader, UE::CaptureManager::ITcpSocketWriter& InWriter);
+	static FUploadVoidResult SerializeData(TArray<uint8> InData, UE::CaptureManager::ITcpSocketWriter& InWriter);
+	static FUploadVoidResult SerializeHash(TStaticArray<uint8, HashSize> InHash, UE::CaptureManager::ITcpSocketWriter& InWriter);
+
+	static FUploadResult<FUploadDataHeader> DeserializeHeader(UE::CaptureManager::ITcpSocketReader& InReader);
+	static FUploadResult<FUploadFileDataHeader> DeserializeFileHeader(UE::CaptureManager::ITcpSocketReader& InReader);
+	static FUploadResult<TArray<uint8>> DeserializeData(uint32 InSize, UE::CaptureManager::ITcpSocketReader& InReader);
+	static FUploadResult<TStaticArray<uint8, HashSize>> DeserializeHash(UE::CaptureManager::ITcpSocketReader& InReader);
+
+private:
+
+	static FUploadVoidResult DeserializeStartHeader(UE::CaptureManager::ITcpSocketReader& InReader);
+	static FUploadVoidResult DeserializeGuid(UE::CaptureManager::ITcpSocketReader& InReader, FGuid& OutFGuid);
+	static FUploadVoidResult DeserializeString(UE::CaptureManager::ITcpSocketReader& InReader, FString& OutTakeName);
+	static FUploadVoidResult DeserializeTotalLength(UE::CaptureManager::ITcpSocketReader& InReader, uint64& OutTotalLength);
+	static FUploadVoidResult DeserializeFileName(UE::CaptureManager::ITcpSocketReader& InReader, FString& OutFileName);
+	static FUploadVoidResult DeserializeLength(UE::CaptureManager::ITcpSocketReader& InReader, uint64& OutLength);
+
+	static constexpr uint32 InactivityTimeout = 15 * 1000; // Miliseconds
+};
diff --git a/Engine/Plugins/VirtualProduction/CaptureManager/CaptureManagerEditor/Source/DataIngestCoreEditor/Public/DataIngestCoreError.h b/Engine/Plugins/VirtualProduction/CaptureManager/CaptureManagerEditor/Source/DataIngestCoreEditor/Public/DataIngestCoreError.h
new file mode 100644
--- /dev/null
+++ b/Engine/Plugins/VirtualProduction/CaptureManager/CaptureManagerEditor/Source/DataIngestCoreEditor/Public/DataIngestCoreError.h
@@ -0,0 +1,38 @@
+// Copyright Epic Games, Inc. All Rights Reserved.
+
+#pragma once
+
+#include "Internationalization/Text.h"
+
+#include "HAL/Platform.h"
+
+/** Asset creation error types. */
+enum class EAssetCreationError
+{
+	InternalError,
+	InvalidArgument,
+	NotFound,
+	Warning,
+};
+
+/** Information related to an error during ingest asset creation. */
+class DATAINGESTCOREEDITOR_API FAssetCreationError
+{
+public:
+
+	FAssetCreationError(FText InMessage, EAssetCreationError InError = EAssetCreationError::InternalError);
+
+	/** Get the error message. */
+	const FText& GetMessage() const;
+
+	/** Get the error type. */
+	EAssetCreationError GetError() const;
+
+private:
+
+	/** Error message. */
+	FText Message;
+
+	/** Error type. */
+	EAssetCreationError Error;
+};
\ No newline at end of file
