Ensure consistency of Capture Data variable names

diff --git a/Engine/Plugins/VirtualProduction/CaptureData/Source/CaptureDataCore/Private/CaptureData.cpp b/Engine/Plugins/VirtualProduction/CaptureData/Source/CaptureDataCore/Private/CaptureData.cpp
--- a/Engine/Plugins/VirtualProduction/CaptureData/Source/CaptureDataCore/Private/CaptureData.cpp
+++ b/Engine/Plugins/VirtualProduction/CaptureData/Source/CaptureDataCore/Private/CaptureData.cpp
@@ -423,7 +423,7 @@ void UFootageCaptureData::GetFrameRanges(const FFrameRate& InTargetRate, ETimeco
 
 	if (bInIncludeAudio)
 	{
-		for (TObjectPtr<class USoundWave> Audio : Audios)
+		for (TObjectPtr<class USoundWave> Audio : AudioTracks)
 		{
 			if (Audio)
 			{
@@ -582,9 +582,9 @@ FTimecode UFootageCaptureData::GetEffectiveAudioTimecode() const
 	FTimecode Timecode;
 	bool TimecodeFromAsset = false;
 
-	if (Audios.Num() > 0)
+	if (AudioTracks.Num() > 0)
 	{
-		const USoundWave* SoundWave = Audios[0].Get();
+		const USoundWave* SoundWave = AudioTracks[0].Get();
 		check(IsValid(SoundWave));
 
 		TOptional<FTimecode> TimecodeOpt = USoundWaveTimecodeUtils::GetTimecode(SoundWave);
@@ -609,9 +609,9 @@ FFrameRate UFootageCaptureData::GetEffectiveAudioTimecodeRate() const
 	FFrameRate FrameRate;
 	bool TimecodeRateFromAsset = false;
 
-	if (Audios.Num() > 0)
+	if (AudioTracks.Num() > 0)
 	{
-		const USoundWave* SoundWave = Audios[0].Get();
+		const USoundWave* SoundWave = AudioTracks[0].Get();
 		check(IsValid(SoundWave));
 
 		TOptional<FFrameRate> FrameRateOpt = USoundWaveTimecodeUtils::GetFrameRate(SoundWave);
@@ -659,9 +659,9 @@ void UFootageCaptureData::GetDefaultTimecodeInfo(FTimecode& OutTimecode, FFrameR
 		}
 	}
 
-	if (Audios.Num() > 0)
+	if (AudioTracks.Num() > 0)
 	{
-		const USoundWave* SoundWave = Audios[0].Get();
+		const USoundWave* SoundWave = AudioTracks[0].Get();
 		check(IsValid(SoundWave));
 
 		TOptional<FTimecode> TimecodeOpt = USoundWaveTimecodeUtils::GetTimecode(SoundWave);
@@ -697,10 +697,19 @@ void UFootageCaptureData::PostLoad()
 
 	if (Audio_DEPRECATED != nullptr)
 	{
-		Audios.Add(MoveTemp(Audio_DEPRECATED));
+		AudioTracks.Add(MoveTemp(Audio_DEPRECATED));
 		Audio_DEPRECATED = nullptr;
 	}
 
+	if (!Audios_DEPRECATED.IsEmpty())
+	{
+		for (TObjectPtr<USoundWave> Audio : Audios_DEPRECATED)
+		{
+			AudioTracks.Add(MoveTemp(Audio));
+		}
+		Audios_DEPRECATED.Empty();
+	}
+
 #if WITH_EDITOR
 	for (FFootageCaptureView& View : Views_DEPRECATED)
 	{
@@ -717,9 +726,9 @@ void UFootageCaptureData::PostLoad()
 		DepthSequences.Add(MoveTemp(View.DepthSequence));
 	}
 #endif
-	if (Audios.Num() > 0 && bAudioTimecodePresent_DEPRECATED == true) 
+	if (AudioTracks.Num() > 0 && bAudioTimecodePresent_DEPRECATED == true)
 	{
-		USoundWave* SoundWave = Audios[0].Get();
+		USoundWave* SoundWave = AudioTracks[0].Get();
 		check(IsValid(SoundWave));
 
 		USoundWaveTimecodeUtils::SetTimecodeInfo(AudioTimecode_DEPRECATED, AudioTimecodeRate_DEPRECATED, SoundWave);
diff --git a/Engine/Plugins/VirtualProduction/CaptureData/Source/CaptureDataCore/Public/CaptureData.h b/Engine/Plugins/VirtualProduction/CaptureData/Source/CaptureDataCore/Public/CaptureData.h
--- a/Engine/Plugins/VirtualProduction/CaptureData/Source/CaptureDataCore/Public/CaptureData.h
+++ b/Engine/Plugins/VirtualProduction/CaptureData/Source/CaptureDataCore/Public/CaptureData.h
@@ -222,7 +222,10 @@ public:
 	TArray<TObjectPtr<class UImgMediaSource>> DepthSequences;
 
 	UPROPERTY(EditAnywhere, BlueprintReadWrite, Category = "Capture")
-	TArray<TObjectPtr<class USoundWave>> Audios;
+	TArray<TObjectPtr<class USoundWave>> AudioTracks;
+
+	UPROPERTY(meta = (DeprecatedProperty, DeprecationMessage = "Audios is deprecated. Please use AudioTracks instead."))
+	TArray<TObjectPtr<class USoundWave>> Audios_DEPRECATED;
 
 	UPROPERTY(meta = (DeprecatedProperty, DeprecationMessage = "Audio is deprecated. Please use Audios instead."))
 	TObjectPtr<class USoundWave> Audio_DEPRECATED;
diff --git a/Engine/Plugins/VirtualProduction/CaptureData/Source/CaptureDataEditor/Private/Customizations/CaptureDataCustomizations.cpp b/Engine/Plugins/VirtualProduction/CaptureData/Source/CaptureDataEditor/Private/Customizations/CaptureDataCustomizations.cpp
--- a/Engine/Plugins/VirtualProduction/CaptureData/Source/CaptureDataEditor/Private/Customizations/CaptureDataCustomizations.cpp
+++ b/Engine/Plugins/VirtualProduction/CaptureData/Source/CaptureDataEditor/Private/Customizations/CaptureDataCustomizations.cpp
@@ -15,7 +15,7 @@ void FFootageCaptureDataCustomization::CustomizeDetails(IDetailLayoutBuilder& In
 {
 	InDetailBuilder.AddPropertyToCategory(InDetailBuilder.GetProperty(GET_MEMBER_NAME_CHECKED(UFootageCaptureData, ImageSequences)));
 	InDetailBuilder.AddPropertyToCategory(InDetailBuilder.GetProperty(GET_MEMBER_NAME_CHECKED(UFootageCaptureData, DepthSequences)));
-	InDetailBuilder.AddPropertyToCategory(InDetailBuilder.GetProperty(GET_MEMBER_NAME_CHECKED(UFootageCaptureData, Audios)));
+	InDetailBuilder.AddPropertyToCategory(InDetailBuilder.GetProperty(GET_MEMBER_NAME_CHECKED(UFootageCaptureData, AudioTracks)));
 	InDetailBuilder.AddPropertyToCategory(InDetailBuilder.GetProperty(GET_MEMBER_NAME_CHECKED(UFootageCaptureData, CameraCalibrations)));
 	InDetailBuilder.AddPropertyToCategory(InDetailBuilder.GetProperty(GET_MEMBER_NAME_CHECKED(UFootageCaptureData, Metadata)));
 
diff --git a/Engine/Plugins/VirtualProduction/CaptureManager/CaptureManagerApp/Source/CaptureDataConverter/Private/Nodes/CaptureValidationNode.cpp b/Engine/Plugins/VirtualProduction/CaptureManager/CaptureManagerApp/Source/CaptureDataConverter/Private/Nodes/CaptureValidationNode.cpp
--- a/Engine/Plugins/VirtualProduction/CaptureManager/CaptureManagerApp/Source/CaptureDataConverter/Private/Nodes/CaptureValidationNode.cpp
+++ b/Engine/Plugins/VirtualProduction/CaptureManager/CaptureManagerApp/Source/CaptureDataConverter/Private/Nodes/CaptureValidationNode.cpp
@@ -101,7 +101,7 @@ FCaptureValidationNode::FResult FCaptureValidationNode::Run()
 		
 		FPaths::MakePathRelativeTo(OutputVideoDirectory, *OutputDirectory);
 
-		IngestData.Videos.Add(Private::ConvertTakeMetadataVideoObject(OutputVideoDirectory, Video));
+		IngestData.Video.Add(Private::ConvertTakeMetadataVideoObject(OutputVideoDirectory, Video));
 	}
 
 	for (const FTakeMetadata::FVideo& Depth : TakeMetadata.Depth)
@@ -119,7 +119,7 @@ FCaptureValidationNode::FResult FCaptureValidationNode::Run()
 
 		FPaths::MakePathRelativeTo(OutputDepthDirectory, *OutputDirectory);
 
-		IngestData.Depths.Add(Private::ConvertTakeMetadataVideoObject(OutputDepthDirectory, Depth));
+		IngestData.Depth.Add(Private::ConvertTakeMetadataVideoObject(OutputDepthDirectory, Depth));
 	}
 
 	for (const FTakeMetadata::FAudio& Audio : TakeMetadata.Audio)
@@ -139,7 +139,7 @@ FCaptureValidationNode::FResult FCaptureValidationNode::Run()
 		OutputAudioDirectory = OutputAudioDirectory / FPaths::SetExtension(AudioParams.AudioFileName, AudioParams.Format);
 		FPaths::MakePathRelativeTo(OutputAudioDirectory, *OutputDirectory);
 
-		IngestData.Audios.Add(Private::ConvertTakeMetadataAudioObject(OutputAudioDirectory, Audio));
+		IngestData.Audio.Add(Private::ConvertTakeMetadataAudioObject(OutputAudioDirectory, Audio));
 	}
 
 	for (const FTakeMetadata::FCalibration& Calibration : TakeMetadata.Calibration)
@@ -157,7 +157,7 @@ FCaptureValidationNode::FResult FCaptureValidationNode::Run()
 
 		FPaths::MakePathRelativeTo(OutputCalibrationFile, *OutputDirectory);
 
-		IngestData.Calibrations.Add(Private::ConvertTakeMetadataCalibrationObject(OutputCalibrationFile, Calibration));
+		IngestData.Calibration.Add(Private::ConvertTakeMetadataCalibrationObject(OutputCalibrationFile, Calibration));
 	}
 
 	IngestData.Version = 1;
diff --git a/Engine/Plugins/VirtualProduction/CaptureManager/CaptureManagerCore/Source/DataIngestCore/Private/IngestCaptureData.cpp b/Engine/Plugins/VirtualProduction/CaptureManager/CaptureManagerCore/Source/DataIngestCore/Private/IngestCaptureData.cpp
--- a/Engine/Plugins/VirtualProduction/CaptureManager/CaptureManagerCore/Source/DataIngestCore/Private/IngestCaptureData.cpp
+++ b/Engine/Plugins/VirtualProduction/CaptureManager/CaptureManagerCore/Source/DataIngestCore/Private/IngestCaptureData.cpp
@@ -173,7 +173,7 @@ FParseResult ParseFile(const FString& InFilePath)
 
 	// Optional
 	const TArray<TSharedPtr<FJsonValue>>* Videos;
-	if (FormatObject->TryGetArrayField(TEXT("Videos"), Videos))
+	if (FormatObject->TryGetArrayField(TEXT("Video"), Videos))
 	{
 		for (const TSharedPtr<FJsonValue>& Video : *Videos)
 		{
@@ -182,13 +182,13 @@ FParseResult ParseFile(const FString& InFilePath)
 			INGEST_CAPTURE_DATA_CHECK_AND_RETURN(Result.HasValue(), Result.StealError());
 
 			FIngestCaptureData::FVideo ParsedVideo = Result.StealValue();
-			IngestCaptureData.Videos.Emplace(MoveTemp(ParsedVideo));
+			IngestCaptureData.Video.Emplace(MoveTemp(ParsedVideo));
 		}
 	}
 
 	// Optional
 	const TArray<TSharedPtr<FJsonValue>>* Depths;
-	if (FormatObject->TryGetArrayField(TEXT("Depths"), Depths))
+	if (FormatObject->TryGetArrayField(TEXT("Depth"), Depths))
 	{
 		for (const TSharedPtr<FJsonValue>& Depth : *Depths)
 		{
@@ -197,13 +197,13 @@ FParseResult ParseFile(const FString& InFilePath)
 			INGEST_CAPTURE_DATA_CHECK_AND_RETURN(Result.HasValue(), Result.StealError());
 
 			FIngestCaptureData::FVideo ParsedDepth = Result.StealValue();
-			IngestCaptureData.Depths.Emplace(MoveTemp(ParsedDepth));
+			IngestCaptureData.Depth.Emplace(MoveTemp(ParsedDepth));
 		}
 	}
 
 	// Optional
 	const TArray<TSharedPtr<FJsonValue>>* Audios;
-	if (FormatObject->TryGetArrayField(TEXT("Audios"), Audios))
+	if (FormatObject->TryGetArrayField(TEXT("Audio"), Audios))
 	{
 		for (const TSharedPtr<FJsonValue>& Audio : *Audios)
 		{
@@ -212,12 +212,12 @@ FParseResult ParseFile(const FString& InFilePath)
 			INGEST_CAPTURE_DATA_CHECK_AND_RETURN(Result.HasValue(), Result.StealError());
 
 			FIngestCaptureData::FAudio ParsedAudio = Result.StealValue();
-			IngestCaptureData.Audios.Emplace(MoveTemp(ParsedAudio));
+			IngestCaptureData.Audio.Emplace(MoveTemp(ParsedAudio));
 		}
 	}
 
 	const TArray<TSharedPtr<FJsonValue>>* Calibrations;
-	if (FormatObject->TryGetArrayField(TEXT("Calibrations"), Calibrations))
+	if (FormatObject->TryGetArrayField(TEXT("Calibration"), Calibrations))
 	{
 		for (const TSharedPtr<FJsonValue>& Calibration : *Calibrations)
 		{
@@ -226,7 +226,7 @@ FParseResult ParseFile(const FString& InFilePath)
 			INGEST_CAPTURE_DATA_CHECK_AND_RETURN(Result.HasValue(), Result.StealError());
 
 			FIngestCaptureData::FCalibration ParsedCalibration = Result.StealValue();
-			IngestCaptureData.Calibrations.Emplace(MoveTemp(ParsedCalibration));
+			IngestCaptureData.Calibration.Emplace(MoveTemp(ParsedCalibration));
 		}
 	}
 
@@ -319,11 +319,11 @@ TOptional<FText> Serialize(const FString& InFilePath, const FString& InFileName,
 	JsonWriter->WriteValue(TEXT("Slate"), InIngestCaptureData.Slate);
 	JsonWriter->WriteValue(TEXT("TakeNumber"), InIngestCaptureData.TakeNumber);
 
-	if (!InIngestCaptureData.Videos.IsEmpty())
+	if (!InIngestCaptureData.Video.IsEmpty())
 	{
-		JsonWriter->WriteArrayStart(TEXT("Videos"));
+		JsonWriter->WriteArrayStart(TEXT("Video"));
 
-		for (const FIngestCaptureData::FVideo& Video : InIngestCaptureData.Videos)
+		for (const FIngestCaptureData::FVideo& Video : InIngestCaptureData.Video)
 		{
 			SerializeVideo(JsonWriter, Video);
 		}
@@ -331,11 +331,11 @@ TOptional<FText> Serialize(const FString& InFilePath, const FString& InFileName,
 		JsonWriter->WriteArrayEnd();
 	}
 
-	if (!InIngestCaptureData.Depths.IsEmpty())
+	if (!InIngestCaptureData.Depth.IsEmpty())
 	{
-		JsonWriter->WriteArrayStart(TEXT("Depths"));
+		JsonWriter->WriteArrayStart(TEXT("Depth"));
 
-		for (const FIngestCaptureData::FVideo& Depth : InIngestCaptureData.Depths)
+		for (const FIngestCaptureData::FVideo& Depth : InIngestCaptureData.Depth)
 		{
 			SerializeVideo(JsonWriter, Depth);
 		}
@@ -343,11 +343,11 @@ TOptional<FText> Serialize(const FString& InFilePath, const FString& InFileName,
 		JsonWriter->WriteArrayEnd();
 	}
 
-	if (!InIngestCaptureData.Audios.IsEmpty())
+	if (!InIngestCaptureData.Audio.IsEmpty())
 	{
-		JsonWriter->WriteArrayStart(TEXT("Audios"));
+		JsonWriter->WriteArrayStart(TEXT("Audio"));
 
-		for (const FIngestCaptureData::FAudio& Audio : InIngestCaptureData.Audios)
+		for (const FIngestCaptureData::FAudio& Audio : InIngestCaptureData.Audio)
 		{
 			SerializeAudio(JsonWriter, Audio);
 		}
@@ -355,11 +355,11 @@ TOptional<FText> Serialize(const FString& InFilePath, const FString& InFileName,
 		JsonWriter->WriteArrayEnd();
 	}
 
-	if (!InIngestCaptureData.Calibrations.IsEmpty())
+	if (!InIngestCaptureData.Calibration.IsEmpty())
 	{
-		JsonWriter->WriteArrayStart(TEXT("Calibrations"));
+		JsonWriter->WriteArrayStart(TEXT("Calibration"));
 
-		for (const FIngestCaptureData::FCalibration& Calibration : InIngestCaptureData.Calibrations)
+		for (const FIngestCaptureData::FCalibration& Calibration : InIngestCaptureData.Calibration)
 		{
 			SerializeCalibration(JsonWriter, Calibration);
 		}
diff --git a/Engine/Plugins/VirtualProduction/CaptureManager/CaptureManagerCore/Source/DataIngestCore/Public/IngestCaptureData.h b/Engine/Plugins/VirtualProduction/CaptureManager/CaptureManagerCore/Source/DataIngestCore/Public/IngestCaptureData.h
--- a/Engine/Plugins/VirtualProduction/CaptureManager/CaptureManagerCore/Source/DataIngestCore/Public/IngestCaptureData.h
+++ b/Engine/Plugins/VirtualProduction/CaptureManager/CaptureManagerCore/Source/DataIngestCore/Public/IngestCaptureData.h
@@ -49,10 +49,10 @@ struct DATAINGESTCORE_API FIngestCaptureData
 	FString Slate;
 	uint32 TakeNumber;
 	
-	TArray<FVideo> Videos;
-	TArray<FVideo> Depths;
-	TArray<FAudio> Audios;
-	TArray<FCalibration> Calibrations;
+	TArray<FVideo> Video;
+	TArray<FVideo> Depth;
+	TArray<FAudio> Audio;
+	TArray<FCalibration> Calibration;
 };
 
 namespace UE::CaptureManager::IngestCaptureData
diff --git a/Engine/Plugins/VirtualProduction/CaptureManager/CaptureManagerEditor/Source/LiveLinkHubWorkerManager/Private/IngestProcess/IngestCaptureDataProcess.cpp b/Engine/Plugins/VirtualProduction/CaptureManager/CaptureManagerEditor/Source/LiveLinkHubWorkerManager/Private/IngestProcess/IngestCaptureDataProcess.cpp
--- a/Engine/Plugins/VirtualProduction/CaptureManager/CaptureManagerEditor/Source/LiveLinkHubWorkerManager/Private/IngestProcess/IngestCaptureDataProcess.cpp
+++ b/Engine/Plugins/VirtualProduction/CaptureManager/CaptureManagerEditor/Source/LiveLinkHubWorkerManager/Private/IngestProcess/IngestCaptureDataProcess.cpp
@@ -91,9 +91,9 @@ TValueOrError<FIngestProcessResult, FText> FIngestCaptureDataProcess::StartInges
 	IngestProcessResult.AssetsData.Add(PrepareAssetsData(IngestProcessResult.TakeIngestPackagePath, InTakeName, IngestCaptureData));
 	
 	IngestProcessResult.CaptureDataTakeInfo.Name = InTakeName;
-	if (!IngestCaptureData.Videos.IsEmpty())
+	if (!IngestCaptureData.Video.IsEmpty())
 	{
-		const FIngestCaptureData::FVideo& FirstVideo = IngestCaptureData.Videos[0];
+		const FIngestCaptureData::FVideo& FirstVideo = IngestCaptureData.Video[0];
 		if (FirstVideo.FrameRate.IsSet())
 		{
 			IngestProcessResult.CaptureDataTakeInfo.FrameRate = FirstVideo.FrameRate.GetValue();
@@ -124,11 +124,11 @@ UE::CaptureManager::FCreateAssetsData FIngestCaptureDataProcess::PrepareAssetsDa
 	CreateAssetData.TakeId = TakeId;
 	CreateAssetData.PackagePath = InTakeIngestDirectory;
 
-	for (int32 Index = 0; Index < InIngestCaptureData.Videos.Num(); ++Index)
+	for (int32 Index = 0; Index < InIngestCaptureData.Video.Num(); ++Index)
 	{
 		FCreateAssetsData::FImageSequenceData ImageSequenceData;
 
-		const FIngestCaptureData::FVideo& Video = InIngestCaptureData.Videos[Index];
+		const FIngestCaptureData::FVideo& Video = InIngestCaptureData.Video[Index];
 
 		FString ImageSequenceNamePrefix = GetAssetNamePrefix(Video.UserId, TEXT("Video"), Index);
 
@@ -142,9 +142,9 @@ UE::CaptureManager::FCreateAssetsData FIngestCaptureDataProcess::PrepareAssetsDa
 		CreateAssetData.ImageSequences.Add(ImageSequenceData);
 	}
 
-	for (int32 Index = 0; Index < InIngestCaptureData.Depths.Num(); ++Index)
+	for (int32 Index = 0; Index < InIngestCaptureData.Depth.Num(); ++Index)
 	{
-		const FIngestCaptureData::FVideo& Depth = InIngestCaptureData.Depths[Index];
+		const FIngestCaptureData::FVideo& Depth = InIngestCaptureData.Depth[Index];
 
 		FString DepthSequenceNamePrefix = GetAssetNamePrefix(Depth.UserId, TEXT("Depth"), Index);
 
@@ -158,9 +158,9 @@ UE::CaptureManager::FCreateAssetsData FIngestCaptureDataProcess::PrepareAssetsDa
 		CreateAssetData.DepthSequences.Add(DepthSequenceData);
 	}
 
-	for (int32 Index = 0; Index < InIngestCaptureData.Audios.Num(); ++Index)
+	for (int32 Index = 0; Index < InIngestCaptureData.Audio.Num(); ++Index)
 	{
-		const FIngestCaptureData::FAudio& Audio = InIngestCaptureData.Audios[Index];
+		const FIngestCaptureData::FAudio& Audio = InIngestCaptureData.Audio[Index];
 
 		FCreateAssetsData::FAudioData AudioData;
 
@@ -175,7 +175,7 @@ UE::CaptureManager::FCreateAssetsData FIngestCaptureDataProcess::PrepareAssetsDa
 		CreateAssetData.AudioClips.Add(AudioData);
 	}
 
-	for (const FIngestCaptureData::FCalibration& Calibration : InIngestCaptureData.Calibrations)
+	for (const FIngestCaptureData::FCalibration& Calibration : InIngestCaptureData.Calibration)
 	{
 		FCreateAssetsData::FCalibrationData CalibrationData;
 
@@ -191,9 +191,9 @@ UE::CaptureManager::FCreateAssetsData FIngestCaptureDataProcess::PrepareAssetsDa
 	
 	uint32 LastDroppedFrameIndex = 0;
 	uint32 FirstDroppedFrameIndex = 0;
-	if (!InIngestCaptureData.Videos.IsEmpty())
+	if (!InIngestCaptureData.Video.IsEmpty())
 	{
-		for (const uint32 DroppedFrameIndex : InIngestCaptureData.Videos[0].DroppedFrames)
+		for (const uint32 DroppedFrameIndex : InIngestCaptureData.Video[0].DroppedFrames)
 		{
 			if (FirstDroppedFrameIndex == 0)
 			{
@@ -222,22 +222,22 @@ UE::CaptureManager::FCreateAssetsData FIngestCaptureDataProcess::PrepareAssetsDa
 
 void FIngestCaptureDataProcess::ConvertPathsToFull(const FString& InTakeStoragePath, FIngestCaptureData& OutIngestCaptureData)
 {
-	for (FIngestCaptureData::FVideo& Video : OutIngestCaptureData.Videos)
+	for (FIngestCaptureData::FVideo& Video : OutIngestCaptureData.Video)
 	{
 		Video.Path = FPaths::ConvertRelativePathToFull(InTakeStoragePath, Video.Path);
 	}
 
-	for (FIngestCaptureData::FVideo& Depth : OutIngestCaptureData.Depths)
+	for (FIngestCaptureData::FVideo& Depth : OutIngestCaptureData.Depth)
 	{
 		Depth.Path = FPaths::ConvertRelativePathToFull(InTakeStoragePath, Depth.Path);
 	}
 
-	for (FIngestCaptureData::FAudio& Audio : OutIngestCaptureData.Audios)
+	for (FIngestCaptureData::FAudio& Audio : OutIngestCaptureData.Audio)
 	{
 		Audio.Path = FPaths::ConvertRelativePathToFull(InTakeStoragePath, Audio.Path);
 	}
 
-	for (FIngestCaptureData::FCalibration& Calibration : OutIngestCaptureData.Calibrations)
+	for (FIngestCaptureData::FCalibration& Calibration : OutIngestCaptureData.Calibration)
 	{
 		Calibration.Path = FPaths::ConvertRelativePathToFull(InTakeStoragePath, Calibration.Path);
 	}
diff --git a/Engine/Plugins/VirtualProduction/CaptureManager/CaptureManagerEditor/Source/LiveLinkHubWorkerManager/Private/LiveLinkHubWorker.cpp b/Engine/Plugins/VirtualProduction/CaptureManager/CaptureManagerEditor/Source/LiveLinkHubWorkerManager/Private/LiveLinkHubWorker.cpp
--- a/Engine/Plugins/VirtualProduction/CaptureManager/CaptureManagerEditor/Source/LiveLinkHubWorkerManager/Private/LiveLinkHubWorker.cpp
+++ b/Engine/Plugins/VirtualProduction/CaptureManager/CaptureManagerEditor/Source/LiveLinkHubWorkerManager/Private/LiveLinkHubWorker.cpp
@@ -346,7 +346,7 @@ bool FLiveLinkHubImportWorker::CreateCaptureAsset(const FString& InAssetPath, co
 		CaptureData->ImageSequences.Reset();
 		CaptureData->DepthSequences.Reset();
 		CaptureData->CameraCalibrations.Reset();
-		CaptureData->Audios.Reset();
+		CaptureData->AudioTracks.Reset();
 
 		for (const FCaptureDataAssetInfo::FImageSequence& ImageSequence : InResult.ImageSequences)
 		{
@@ -364,7 +364,7 @@ bool FLiveLinkHubImportWorker::CreateCaptureAsset(const FString& InAssetPath, co
 
 		for (const FCaptureDataAssetInfo::FAudio& Audio : InResult.Audios)
 		{
-			CaptureData->Audios.Add(Audio.Asset);
+			CaptureData->AudioTracks.Add(Audio.Asset);
 		}
 
 		for (const FCaptureDataAssetInfo::FCalibration& Calibration : InResult.Calibrations)
