[MetaHuman] Resolved issue with monolithic exes

diff --git a/Engine/Plugins/VirtualProduction/CaptureManager/CaptureManagerCore/Source/CaptureManagerTakeMetadata/Private/Tests/CaptureManagerTakeMetadata.spec.cpp b/Engine/Plugins/VirtualProduction/CaptureManager/CaptureManagerCore/Source/CaptureManagerTakeMetadata/Private/Tests/CaptureManagerTakeMetadata.spec.cpp
--- a/Engine/Plugins/VirtualProduction/CaptureManager/CaptureManagerCore/Source/CaptureManagerTakeMetadata/Private/Tests/CaptureManagerTakeMetadata.spec.cpp
+++ b/Engine/Plugins/VirtualProduction/CaptureManager/CaptureManagerCore/Source/CaptureManagerTakeMetadata/Private/Tests/CaptureManagerTakeMetadata.spec.cpp
@@ -9,6 +9,21 @@
 
 #if WITH_AUTOMATION_TESTS
 
+namespace CaptureManagerTakeMetadata_Spec
+{
+FString& GetJsonTestsDir()
+{
+	static FString JsonTestsDir = IPluginManager::Get().FindPlugin(UE_PLUGIN_NAME)->GetContentDir() / TEXT("TakeMetadata") / TEXT("TestInputs");
+	return JsonTestsDir;
+}
+
+FTakeMetadataParser& GetTakeMetadataParser()
+{
+	static FTakeMetadataParser TakeMetadataParser;
+	return TakeMetadataParser;
+}
+}
+
 BEGIN_DEFINE_SPEC(TestTakeMetadataSchemaSpec, TEXT("TestTakeMetadataSchema"), EAutomationTestFlags::EngineFilter | EAutomationTestFlags_ApplicationContextMask)
 
 struct TestSpecification
@@ -18,20 +33,18 @@ struct TestSpecification
 	TArray<FString> ExpectedMessageKeywords;
 };
 
-
-FString JsonTestsDir = IPluginManager::Get().FindPlugin(UE_PLUGIN_NAME)->GetContentDir() / TEXT("TakeMetadata") / TEXT("TestInputs");
-FTakeMetadataParser TakeMetadataParser;
-
 END_DEFINE_SPEC(TestTakeMetadataSchemaSpec)
 
 void TestTakeMetadataSchemaSpec::Define()
 {
+	using namespace CaptureManagerTakeMetadata_Spec;
+
 	Describe(TEXT("TestTakeMetadataSchemaTests"), [this]
 	{
 			It(TEXT("take_valid_v4_0_mandatory"), [this]
 				{
 					TValueOrError<FTakeMetadata, FTakeMetadataParserError> Result =
-						TakeMetadataParser.Parse(JsonTestsDir / TEXT("take_valid_v4_0_mandatory.json"));
+						GetTakeMetadataParser().Parse(GetJsonTestsDir() / TEXT("take_valid_v4_0_mandatory.json"));
 
 					UTEST_TRUE(TEXT("Should have value."), Result.HasValue());
 
@@ -54,7 +67,7 @@ void TestTakeMetadataSchemaSpec::Define()
 			It(TEXT("take_valid_v4_0_mandatory_media"), [this]
 				{
 					TValueOrError<FTakeMetadata, FTakeMetadataParserError> Result =
-						TakeMetadataParser.Parse(JsonTestsDir / TEXT("take_valid_v4_0_mandatory_media.json"));
+						GetTakeMetadataParser().Parse(GetJsonTestsDir() / TEXT("take_valid_v4_0_mandatory_media.json"));
 
 					UTEST_TRUE(TEXT("Should have value."), Result.HasValue());
 
@@ -92,7 +105,7 @@ void TestTakeMetadataSchemaSpec::Define()
 			It(TEXT("take_valid_v4_0"), [this]
 				{
 					TValueOrError<FTakeMetadata, FTakeMetadataParserError> Result =
-						TakeMetadataParser.Parse(JsonTestsDir / TEXT("take_valid_v4_0.json"));
+						GetTakeMetadataParser().Parse(GetJsonTestsDir() / TEXT("take_valid_v4_0.json"));
 
 					UTEST_TRUE(TEXT("Should have value."), Result.HasValue());
 
@@ -158,7 +171,7 @@ void TestTakeMetadataSchemaSpec::Define()
 			It(TEXT("take_valid_v3_0"), [this]
 				{
 					TValueOrError<FTakeMetadata, FTakeMetadataParserError> Result =
-						TakeMetadataParser.Parse(JsonTestsDir / TEXT("take_valid_v3_0.json"));
+						GetTakeMetadataParser().Parse(GetJsonTestsDir() / TEXT("take_valid_v3_0.json"));
 
 					UTEST_TRUE(TEXT("Should have value."), Result.HasValue());
 
@@ -224,7 +237,7 @@ void TestTakeMetadataSchemaSpec::Define()
 		It(TEXT("take_valid_v2_0"), [this]
 			{
 				TValueOrError<FTakeMetadata, FTakeMetadataParserError> Result =
-				TakeMetadataParser.Parse(JsonTestsDir / TEXT("take_valid_v2_0.json"));
+					GetTakeMetadataParser().Parse(GetJsonTestsDir() / TEXT("take_valid_v2_0.json"));
 
 				UTEST_TRUE(TEXT("Should have value."), Result.HasValue());
 
@@ -288,7 +301,7 @@ void TestTakeMetadataSchemaSpec::Define()
 
 		It(TEXT("take_valid_v1_0_calibration"), [this]
 			{
-				TValueOrError<FTakeMetadata, FTakeMetadataParserError> Result = TakeMetadataParser.Parse(JsonTestsDir / TEXT("take_valid_v1_0_calibration.json"));
+				TValueOrError<FTakeMetadata, FTakeMetadataParserError> Result = GetTakeMetadataParser().Parse(GetJsonTestsDir() / TEXT("take_valid_v1_0_calibration.json"));
 				TestTrue(TEXT("Should have value."), Result.HasValue());
 				if (Result.HasValue())
 				{
@@ -332,7 +345,7 @@ void TestTakeMetadataSchemaSpec::Define()
 		{
 			It(Test.Name, [this, Test]
 				{
-					TValueOrError<FTakeMetadata, FTakeMetadataParserError> Result = TakeMetadataParser.Parse(JsonTestsDir / Test.Name + ".json");
+					TValueOrError<FTakeMetadata, FTakeMetadataParserError> Result = GetTakeMetadataParser().Parse(GetJsonTestsDir() / Test.Name + ".json");
 					if (!TestTrue(TEXT("Should have error."), Result.HasError()))
 					{
 						return;
