[MetaHuman] Move test specification into namespace

diff --git a/Engine/Plugins/VirtualProduction/CaptureManager/CaptureManagerCore/Source/CaptureManagerTakeMetadata/Private/Tests/CaptureManagerTakeMetadata.spec.cpp b/Engine/Plugins/VirtualProduction/CaptureManager/CaptureManagerCore/Source/CaptureManagerTakeMetadata/Private/Tests/CaptureManagerTakeMetadata.spec.cpp
--- a/Engine/Plugins/VirtualProduction/CaptureManager/CaptureManagerCore/Source/CaptureManagerTakeMetadata/Private/Tests/CaptureManagerTakeMetadata.spec.cpp
+++ b/Engine/Plugins/VirtualProduction/CaptureManager/CaptureManagerCore/Source/CaptureManagerTakeMetadata/Private/Tests/CaptureManagerTakeMetadata.spec.cpp
@@ -22,7 +22,6 @@ FTakeMetadataParser& GetTakeMetadataParser()
 	static FTakeMetadataParser TakeMetadataParser;
 	return TakeMetadataParser;
 }
-}
 
 BEGIN_DEFINE_SPEC(TestTakeMetadataSchemaSpec, TEXT("TestTakeMetadataSchema"), EAutomationTestFlags::EngineFilter | EAutomationTestFlags_ApplicationContextMask)
 
@@ -34,8 +33,10 @@ struct TestSpecification
 };
 
 END_DEFINE_SPEC(TestTakeMetadataSchemaSpec)
+}
+
 
-void TestTakeMetadataSchemaSpec::Define()
+void CaptureManagerTakeMetadata_Spec::TestTakeMetadataSchemaSpec::Define()
 {
 	using namespace CaptureManagerTakeMetadata_Spec;
 
