Avoid blocking the main thread when shutting down ControlCommunication by closing the TcpClient before the runnable.

#rnx

diff --git a/Engine/Plugins/VirtualProduction/CaptureManager/CaptureManagerCore/Source/CaptureUtils/Private/Network/TcpClient.cpp b/Engine/Plugins/VirtualProduction/CaptureManager/CaptureManagerCore/Source/CaptureUtils/Private/Network/TcpClient.cpp
--- a/Engine/Plugins/VirtualProduction/CaptureManager/CaptureManagerCore/Source/CaptureUtils/Private/Network/TcpClient.cpp
+++ b/Engine/Plugins/VirtualProduction/CaptureManager/CaptureManagerCore/Source/CaptureUtils/Private/Network/TcpClient.cpp
@@ -86,7 +86,6 @@ TProtocolResult<void> FTcpClient::Stop()
 	{
 		TcpSocket->Shutdown(ESocketShutdownMode::ReadWrite);
 		TcpSocket->Close();
-		TcpSocket = nullptr;
 	}
 
 	bRunning = false;
