[CaptureManager] Improve metadata key names

diff --git a/Engine/Plugins/VirtualProduction/CaptureManager/CaptureManagerCore/Source/DataIngestCore/Private/IngestCaptureData.cpp b/Engine/Plugins/VirtualProduction/CaptureManager/CaptureManagerCore/Source/DataIngestCore/Private/IngestCaptureData.cpp
--- a/Engine/Plugins/VirtualProduction/CaptureManager/CaptureManagerCore/Source/DataIngestCore/Private/IngestCaptureData.cpp
+++ b/Engine/Plugins/VirtualProduction/CaptureManager/CaptureManagerCore/Source/DataIngestCore/Private/IngestCaptureData.cpp
@@ -172,10 +172,10 @@ FParseResult ParseFile(const FString& InFilePath)
 										 LOCTEXT("Parse_MissingTakeNumber", "Json doesn't contain Take Number field"));
 
 	// Optional
-	const TArray<TSharedPtr<FJsonValue>>* Videos;
-	if (FormatObject->TryGetArrayField(TEXT("Video"), Videos))
+	const TArray<TSharedPtr<FJsonValue>>* VideoArray;
+	if (FormatObject->TryGetArrayField(TEXT("Video"), VideoArray))
 	{
-		for (const TSharedPtr<FJsonValue>& Video : *Videos)
+		for (const TSharedPtr<FJsonValue>& Video : *VideoArray)
 		{
 			TValueOrError<FIngestCaptureData::FVideo, FText> Result = ParseVideoObject(Video);
 
@@ -187,10 +187,10 @@ FParseResult ParseFile(const FString& InFilePath)
 	}
 
 	// Optional
-	const TArray<TSharedPtr<FJsonValue>>* Depths;
-	if (FormatObject->TryGetArrayField(TEXT("Depth"), Depths))
+	const TArray<TSharedPtr<FJsonValue>>* DepthArray;
+	if (FormatObject->TryGetArrayField(TEXT("Depth"), DepthArray))
 	{
-		for (const TSharedPtr<FJsonValue>& Depth : *Depths)
+		for (const TSharedPtr<FJsonValue>& Depth : *DepthArray)
 		{
 			TValueOrError<FIngestCaptureData::FVideo, FText> Result = ParseVideoObject(Depth);
 
@@ -202,10 +202,10 @@ FParseResult ParseFile(const FString& InFilePath)
 	}
 
 	// Optional
-	const TArray<TSharedPtr<FJsonValue>>* Audios;
-	if (FormatObject->TryGetArrayField(TEXT("Audio"), Audios))
+	const TArray<TSharedPtr<FJsonValue>>* AudioArray;
+	if (FormatObject->TryGetArrayField(TEXT("Audio"), AudioArray))
 	{
-		for (const TSharedPtr<FJsonValue>& Audio : *Audios)
+		for (const TSharedPtr<FJsonValue>& Audio : *AudioArray)
 		{
 			TValueOrError<FIngestCaptureData::FAudio, FText> Result = ParseAudioObject(Audio);
 
@@ -216,10 +216,10 @@ FParseResult ParseFile(const FString& InFilePath)
 		}
 	}
 
-	const TArray<TSharedPtr<FJsonValue>>* Calibrations;
-	if (FormatObject->TryGetArrayField(TEXT("Calibration"), Calibrations))
+	const TArray<TSharedPtr<FJsonValue>>* CalibrationArray;
+	if (FormatObject->TryGetArrayField(TEXT("Calibration"), CalibrationArray))
 	{
-		for (const TSharedPtr<FJsonValue>& Calibration : *Calibrations)
+		for (const TSharedPtr<FJsonValue>& Calibration : *CalibrationArray)
 		{
 			TValueOrError<FIngestCaptureData::FCalibration, FText> Result = ParseCalibrationObject(Calibration);
 
