[CaptureManager] Sanitize multiple slashes in package paths

diff --git a/Engine/Plugins/VirtualProduction/CaptureManager/CaptureManagerCore/Source/CaptureUtils/Private/Asset/CaptureAssetSanitization.cpp b/Engine/Plugins/VirtualProduction/CaptureManager/CaptureManagerCore/Source/CaptureUtils/Private/Asset/CaptureAssetSanitization.cpp
new file mode 100644
--- /dev/null
+++ b/Engine/Plugins/VirtualProduction/CaptureManager/CaptureManagerCore/Source/CaptureUtils/Private/Asset/CaptureAssetSanitization.cpp
@@ -0,0 +1,49 @@
+// Copyright Epic Games, Inc. All Rights Reserved.
+
+#include "Asset/CaptureAssetSanitization.h"
+
+#include "Internationalization/Text.h"
+
+#include "Logging/LogMacros.h"
+#include "Misc/Paths.h"
+
+DEFINE_LOG_CATEGORY_STATIC(LogCaptureUtilsAssetValidation, Log, All);
+
+namespace UE::CaptureManager::Private
+{
+
+static void Sanitize(FString& OutPath, const FString& InInvalidChars, const FString::ElementType InReplaceWith)
+{
+	// Make sure the package path only contains valid characters
+	FText OutError;
+	if (FName::IsValidXName(OutPath, InInvalidChars, &OutError))
+	{
+		return;
+	}
+	UE_LOG(LogCaptureUtilsAssetValidation, Display, TEXT("%s"), *OutError.ToString());
+
+	for (const FString::ElementType& InvalidChar : InInvalidChars)
+	{
+		OutPath.ReplaceCharInline(InvalidChar, InReplaceWith, ESearchCase::CaseSensitive);
+	}
+
+	UE_LOG(LogCaptureUtilsAssetValidation, Display, TEXT("Sanitized path: %s"), *OutPath);
+}
+
+}
+
+namespace UE::CaptureManager
+{
+
+void SanitizePackagePath(FString& OutPath, const FString::ElementType InReplaceWith)
+{	
+	FPaths::RemoveDuplicateSlashes(OutPath);
+	return Private::Sanitize(OutPath, INVALID_OBJECTPATH_CHARACTERS INVALID_LONGPACKAGE_CHARACTERS, InReplaceWith);
+}
+
+void SanitizeAssetName(FString& OutPath, const FString::ElementType InReplaceWith)
+{
+	return Private::Sanitize(OutPath, INVALID_OBJECTNAME_CHARACTERS, InReplaceWith);
+}
+
+}
\ No newline at end of file
diff --git a/Engine/Plugins/VirtualProduction/CaptureManager/CaptureManagerCore/Source/CaptureUtils/Private/Asset/Sanitization.cpp b/Engine/Plugins/VirtualProduction/CaptureManager/CaptureManagerCore/Source/CaptureUtils/Private/Asset/Sanitization.cpp
deleted file mode 100644
diff --git a/Engine/Plugins/VirtualProduction/CaptureManager/CaptureManagerCore/Source/CaptureUtils/Public/Asset/CaptureAssetSanitization.h b/Engine/Plugins/VirtualProduction/CaptureManager/CaptureManagerCore/Source/CaptureUtils/Public/Asset/CaptureAssetSanitization.h
new file mode 100644
--- /dev/null
+++ b/Engine/Plugins/VirtualProduction/CaptureManager/CaptureManagerCore/Source/CaptureUtils/Public/Asset/CaptureAssetSanitization.h
@@ -0,0 +1,15 @@
+// Copyright Epic Games, Inc. All Rights Reserved.
+
+#pragma once
+
+#include "Containers/UnrealString.h"
+
+namespace UE::CaptureManager
+{
+// Replaces any invalid package path characters 
+CAPTUREUTILS_API void SanitizePackagePath(FString& OutPath, FString::ElementType InReplaceWith = TEXT('_'));
+
+// Replaces any invalid object name characters 
+CAPTUREUTILS_API void SanitizeAssetName(FString& OutPath, FString::ElementType InReplaceWith = TEXT('_'));
+
+}
diff --git a/Engine/Plugins/VirtualProduction/CaptureManager/CaptureManagerCore/Source/CaptureUtils/Public/Asset/Sanitization.h b/Engine/Plugins/VirtualProduction/CaptureManager/CaptureManagerCore/Source/CaptureUtils/Public/Asset/Sanitization.h
deleted file mode 100644
diff --git a/Engine/Plugins/VirtualProduction/CaptureManager/CaptureManagerEditor/Source/LiveLinkHubWorkerManager/Private/IngestProcess/IngestCaptureDataProcess.cpp b/Engine/Plugins/VirtualProduction/CaptureManager/CaptureManagerEditor/Source/LiveLinkHubWorkerManager/Private/IngestProcess/IngestCaptureDataProcess.cpp
--- a/Engine/Plugins/VirtualProduction/CaptureManager/CaptureManagerEditor/Source/LiveLinkHubWorkerManager/Private/IngestProcess/IngestCaptureDataProcess.cpp
+++ b/Engine/Plugins/VirtualProduction/CaptureManager/CaptureManagerEditor/Source/LiveLinkHubWorkerManager/Private/IngestProcess/IngestCaptureDataProcess.cpp
@@ -7,7 +7,7 @@
 #include "Utils/ParseTakeUtils.h"
 #include "Utils/UnrealCalibrationParser.h"
 
-#include "Asset/Sanitization.h"
+#include "Asset/CaptureAssetSanitization.h"
 #include "Async/HelperFunctions.h"
 
 #include "Editor.h"
