[CaptureManager] Replace UserId with Name in take metadata

diff --git a/Engine/Plugins/VirtualProduction/CaptureManager/CaptureManagerApp/Source/CaptureDataConverter/Private/Nodes/CaptureConvertAudioData.cpp b/Engine/Plugins/VirtualProduction/CaptureManager/CaptureManagerApp/Source/CaptureDataConverter/Private/Nodes/CaptureConvertAudioData.cpp
--- a/Engine/Plugins/VirtualProduction/CaptureManager/CaptureManagerApp/Source/CaptureDataConverter/Private/Nodes/CaptureConvertAudioData.cpp
+++ b/Engine/Plugins/VirtualProduction/CaptureManager/CaptureManagerApp/Source/CaptureDataConverter/Private/Nodes/CaptureConvertAudioData.cpp
@@ -47,7 +47,7 @@ FCaptureConvertAudioData::FResult FCaptureConvertAudioData::CopyAudioFile()
 
 	FTaskProgress::FTask Task = Params.TaskProgress->StartTask();
 
-	FString DestinationDirectory = OutputDirectory / Audio.UserId;
+	FString DestinationDirectory = OutputDirectory / Audio.Name;
 	const FString AudioFilePath = FPaths::ConvertRelativePathToFull(Params.TakeOriginDirectory, Audio.Path);
 
 	FCopyProgressReporter ProgressReporter(Task, Params.StopToken);
@@ -75,7 +75,7 @@ FCaptureConvertAudioData::FResult FCaptureConvertAudioData::ConvertAudioFile()
 
 	FTaskProgress::FTask Task = Params.TaskProgress->StartTask();
 
-	FString DestinationDirectory = OutputDirectory / Audio.UserId;
+	FString DestinationDirectory = OutputDirectory / Audio.Name;
 	const FString AudioFilePath = FPaths::ConvertRelativePathToFull(Params.TakeOriginDirectory, Audio.Path);
 
 	FCaptureManagerMediaRWModule& MediaRWModule =
diff --git a/Engine/Plugins/VirtualProduction/CaptureManager/CaptureManagerApp/Source/CaptureDataConverter/Private/Nodes/CaptureConvertCalibrationData.cpp b/Engine/Plugins/VirtualProduction/CaptureManager/CaptureManagerApp/Source/CaptureDataConverter/Private/Nodes/CaptureConvertCalibrationData.cpp
--- a/Engine/Plugins/VirtualProduction/CaptureManager/CaptureManagerApp/Source/CaptureDataConverter/Private/Nodes/CaptureConvertCalibrationData.cpp
+++ b/Engine/Plugins/VirtualProduction/CaptureManager/CaptureManagerApp/Source/CaptureDataConverter/Private/Nodes/CaptureConvertCalibrationData.cpp
@@ -45,7 +45,7 @@ FCaptureConvertCalibrationData::FResult FCaptureConvertCalibrationData::CopyCali
 
 	FTaskProgress::FTask Task = Params.TaskProgress->StartTask();
 
-	FString DestinationDirectory = OutputDirectory / Calibration.UserId;
+	FString DestinationDirectory = OutputDirectory / Calibration.Name;
 	const FString CalibrationFilePath = FPaths::ConvertRelativePathToFull(Params.TakeOriginDirectory, Calibration.Path);
 
 	FCopyProgressReporter ProgressReporter(Task, Params.StopToken);
@@ -81,7 +81,7 @@ FCaptureConvertCalibrationData::FResult FCaptureConvertCalibrationData::ConvertC
 	FTaskProgress::FTask Task = Params.TaskProgress->StartTask();
 
 	const FString CalibrationFilePath = FPaths::ConvertRelativePathToFull(Params.TakeOriginDirectory, Calibration.Path);
-	FString TargetDirectory = OutputDirectory / Calibration.UserId;
+	FString TargetDirectory = OutputDirectory / Calibration.Name;
 
 	FCaptureManagerMediaRWModule& MediaRWModule =
 		FModuleManager::LoadModuleChecked<FCaptureManagerMediaRWModule>("CaptureManagerMediaRW");
diff --git a/Engine/Plugins/VirtualProduction/CaptureManager/CaptureManagerApp/Source/CaptureDataConverter/Private/Nodes/CaptureConvertDepthData.cpp b/Engine/Plugins/VirtualProduction/CaptureManager/CaptureManagerApp/Source/CaptureDataConverter/Private/Nodes/CaptureConvertDepthData.cpp
--- a/Engine/Plugins/VirtualProduction/CaptureManager/CaptureManagerApp/Source/CaptureDataConverter/Private/Nodes/CaptureConvertDepthData.cpp
+++ b/Engine/Plugins/VirtualProduction/CaptureManager/CaptureManagerApp/Source/CaptureDataConverter/Private/Nodes/CaptureConvertDepthData.cpp
@@ -68,7 +68,7 @@ FCaptureConvertDepthData::FResult FCaptureConvertDepthData::ConvertData()
 	FTaskProgress::FTask Task = Params.TaskProgress->StartTask();
 
 	const FString DepthFilePath = FPaths::ConvertRelativePathToFull(Params.TakeOriginDirectory, Depth.Path);
-	FString TargetDirectory = OutputDirectory / Depth.UserId;
+	FString TargetDirectory = OutputDirectory / Depth.Name;
 
 	FCaptureManagerMediaRWModule& MediaRWModule =
 		FModuleManager::LoadModuleChecked<FCaptureManagerMediaRWModule>("CaptureManagerMediaRW");
@@ -151,7 +151,7 @@ FCaptureConvertDepthData::FResult FCaptureConvertDepthData::CopyData()
 		Task.Update(InProgress);
 	}));
 
-	FString DestinationDirectory = OutputDirectory / Depth.UserId;
+	FString DestinationDirectory = OutputDirectory / Depth.Name;
 	FCaptureConvertDepthData::FResult Result = MakeValue();
 	IImageWrapperModule& ImageWrapperModule = FModuleManager::Get().LoadModuleChecked<IImageWrapperModule>(TEXT("ImageWrapper"));
 	FileManager.IterateDirectoryRecursively(*DepthFolderPath, [this, &TaskProgress, &Result, DestinationDirectory, &FileManager, &ImageWrapperModule](const TCHAR* InFileName, bool bInIsDirectory)
diff --git a/Engine/Plugins/VirtualProduction/CaptureManager/CaptureManagerApp/Source/CaptureDataConverter/Private/Nodes/CaptureConvertVideoData.cpp b/Engine/Plugins/VirtualProduction/CaptureManager/CaptureManagerApp/Source/CaptureDataConverter/Private/Nodes/CaptureConvertVideoData.cpp
--- a/Engine/Plugins/VirtualProduction/CaptureManager/CaptureManagerApp/Source/CaptureDataConverter/Private/Nodes/CaptureConvertVideoData.cpp
+++ b/Engine/Plugins/VirtualProduction/CaptureManager/CaptureManagerApp/Source/CaptureDataConverter/Private/Nodes/CaptureConvertVideoData.cpp
@@ -62,7 +62,7 @@ FCaptureConvertVideoData::FResult FCaptureConvertVideoData::ConvertData()
 	FTaskProgress::FTask Task = Params.TaskProgress->StartTask();
 
 	const FString VideoFilePath = FPaths::ConvertRelativePathToFull(Params.TakeOriginDirectory, Video.Path);
-	FString DestinationDirectory = OutputDirectory / Video.UserId;
+	FString DestinationDirectory = OutputDirectory / Video.Name;
 
 	FCaptureManagerMediaRWModule& MediaRWModule =
 		FModuleManager::LoadModuleChecked<FCaptureManagerMediaRWModule>("CaptureManagerMediaRW");
@@ -138,7 +138,7 @@ FCaptureConvertVideoData::FResult FCaptureConvertVideoData::CopyData()
 		Task.Update(InProgress);
 	}));
 
-	FString DestinationDirectory = OutputDirectory / Video.UserId;
+	FString DestinationDirectory = OutputDirectory / Video.Name;
 
 	FCaptureConvertVideoData::FResult Result = MakeValue();
 
diff --git a/Engine/Plugins/VirtualProduction/CaptureManager/CaptureManagerApp/Source/CaptureDataConverter/Private/Nodes/CaptureValidationNode.cpp b/Engine/Plugins/VirtualProduction/CaptureManager/CaptureManagerApp/Source/CaptureDataConverter/Private/Nodes/CaptureValidationNode.cpp
--- a/Engine/Plugins/VirtualProduction/CaptureManager/CaptureManagerApp/Source/CaptureDataConverter/Private/Nodes/CaptureValidationNode.cpp
+++ b/Engine/Plugins/VirtualProduction/CaptureManager/CaptureManagerApp/Source/CaptureDataConverter/Private/Nodes/CaptureValidationNode.cpp
@@ -22,7 +22,7 @@ FIngestCaptureData::FVideo ConvertTakeMetadataVideoObject(const FString& InVideo
 {
 	FIngestCaptureData::FVideo Video;
 
-	Video.UserId = InVideo.UserId;
+	Video.Name = InVideo.Name;
 	Video.Path = InVideoDirectory;
 	Video.FrameRate = InVideo.FrameRate;
 	Video.FrameWidth = InVideo.FrameWidth;
@@ -37,7 +37,7 @@ FIngestCaptureData::FAudio ConvertTakeMetadataAudioObject(const FString& InAudio
 {
 	FIngestCaptureData::FAudio Audio;
 
-	Audio.UserId = InAudio.UserId;
+	Audio.Name = InAudio.Name;
 	Audio.Path = InAudioDirectory;
 	Audio.TimecodeStart = InAudio.TimecodeStart;
 	Audio.TimecodeRate = InAudio.TimecodeRate;
@@ -49,7 +49,7 @@ FIngestCaptureData::FCalibration ConvertTakeMetadataCalibrationObject(const FStr
 {
 	FIngestCaptureData::FCalibration Calibration;
 
-	Calibration.UserId = InCalibration.UserId;
+	Calibration.Name = InCalibration.Name;
 	Calibration.Path = InCalibrationPath;
 	
 	return Calibration;
@@ -89,7 +89,7 @@ FCaptureValidationNode::FResult FCaptureValidationNode::Run()
 	{
 		static const FString VideoDirectory = TEXT("Video");
 
-		FString OutputVideoDirectory = OutputDirectory / VideoDirectory / Video.UserId;
+		FString OutputVideoDirectory = OutputDirectory / VideoDirectory / Video.Name;
 
 		FCaptureValidationNode::FResult Result = CheckImages(OutputVideoDirectory, {});
 
@@ -108,7 +108,7 @@ FCaptureValidationNode::FResult FCaptureValidationNode::Run()
 	{
 		static const FString DepthDirectory = TEXT("Depth");
 
-		FString OutputDepthDirectory = OutputDirectory / DepthDirectory / Depth.UserId;
+		FString OutputDepthDirectory = OutputDirectory / DepthDirectory / Depth.Name;
 
 		FCaptureValidationNode::FResult Result = CheckImages(OutputDepthDirectory, EImageFormat::EXR);
 
@@ -126,7 +126,7 @@ FCaptureValidationNode::FResult FCaptureValidationNode::Run()
 	{
 		static const FString AudioDirectory = TEXT("Audio");
 
-		FString OutputAudioDirectory = OutputDirectory / AudioDirectory / Audio.UserId;
+		FString OutputAudioDirectory = OutputDirectory / AudioDirectory / Audio.Name;
 
 		FCaptureConvertAudioOutputParams AudioParams = Params.AudioOutputParams.GetValue();
 		FCaptureValidationNode::FResult Result = CheckAudio(AudioParams.AudioFileName, OutputAudioDirectory);
@@ -147,7 +147,7 @@ FCaptureValidationNode::FResult FCaptureValidationNode::Run()
 		static const FString CalibrationDirectory = TEXT("Calibration");
 
 		FCaptureConvertCalibrationOutputParams CalibrationParams = Params.CalibrationOutputParams.GetValue();
-		FString OutputCalibrationFile = OutputDirectory / CalibrationDirectory / Calibration.UserId / CalibrationParams.FileName;
+		FString OutputCalibrationFile = OutputDirectory / CalibrationDirectory / Calibration.Name / CalibrationParams.FileName;
 		OutputCalibrationFile = FPaths::SetExtension(OutputCalibrationFile, TEXT("json"));
 
 		if (!IFileManager::Get().FileExists(*OutputCalibrationFile))
diff --git a/Engine/Plugins/VirtualProduction/CaptureManager/CaptureManagerApp/Source/CaptureDataConverter/Private/Nodes/ThirdPartyEncoder/CaptureConvertAudioDataThirdParty.cpp b/Engine/Plugins/VirtualProduction/CaptureManager/CaptureManagerApp/Source/CaptureDataConverter/Private/Nodes/ThirdPartyEncoder/CaptureConvertAudioDataThirdParty.cpp
--- a/Engine/Plugins/VirtualProduction/CaptureManager/CaptureManagerApp/Source/CaptureDataConverter/Private/Nodes/ThirdPartyEncoder/CaptureConvertAudioDataThirdParty.cpp
+++ b/Engine/Plugins/VirtualProduction/CaptureManager/CaptureManagerApp/Source/CaptureDataConverter/Private/Nodes/ThirdPartyEncoder/CaptureConvertAudioDataThirdParty.cpp
@@ -60,7 +60,7 @@ FCaptureConvertAudioDataThirdParty::FResult FCaptureConvertAudioDataThirdParty::
 
 	FTaskProgress::FTask Task = Params.TaskProgress->StartTask();
 
-	FString DestinationDirectory = OutputDirectory / Audio.UserId;
+	FString DestinationDirectory = OutputDirectory / Audio.Name;
 	const FString AudioFilePath = FPaths::ConvertRelativePathToFull(Params.TakeOriginDirectory, Audio.Path);
 
 	FCopyProgressReporter ProgressReporter(Task, Params.StopToken);
@@ -93,7 +93,7 @@ FCaptureConvertAudioDataThirdParty::FResult FCaptureConvertAudioDataThirdParty::
 
 	FTaskProgress::FTask Task = Params.TaskProgress->StartTask();
 
-	FString DestinationDirectory = OutputDirectory / Audio.UserId;
+	FString DestinationDirectory = OutputDirectory / Audio.Name;
 	const FString AudioFilePath = FPaths::ConvertRelativePathToFull(Params.TakeOriginDirectory, Audio.Path);
 
 	const FString AudioOutputFile = FPaths::SetExtension(FPaths::Combine(DestinationDirectory, AudioParams.AudioFileName), AudioParams.Format);
diff --git a/Engine/Plugins/VirtualProduction/CaptureManager/CaptureManagerApp/Source/CaptureDataConverter/Private/Nodes/ThirdPartyEncoder/CaptureConvertVideoDataThirdParty.cpp b/Engine/Plugins/VirtualProduction/CaptureManager/CaptureManagerApp/Source/CaptureDataConverter/Private/Nodes/ThirdPartyEncoder/CaptureConvertVideoDataThirdParty.cpp
--- a/Engine/Plugins/VirtualProduction/CaptureManager/CaptureManagerApp/Source/CaptureDataConverter/Private/Nodes/ThirdPartyEncoder/CaptureConvertVideoDataThirdParty.cpp
+++ b/Engine/Plugins/VirtualProduction/CaptureManager/CaptureManagerApp/Source/CaptureDataConverter/Private/Nodes/ThirdPartyEncoder/CaptureConvertVideoDataThirdParty.cpp
@@ -160,7 +160,7 @@ FCaptureConvertVideoDataThirdParty::FResult FCaptureConvertVideoDataThirdParty::
 	FTaskProgress::FTask Task = Params.TaskProgress->StartTask();
 
 	const FString VideoFilePath = FPaths::ConvertRelativePathToFull(Params.TakeOriginDirectory, Video.Path);
-	FString DestinationDirectory = OutputDirectory / Video.UserId;
+	FString DestinationDirectory = OutputDirectory / Video.Name;
 
 	const FString ImageFileName = FString::Format(TEXT("{0}_%06d.{1}"), { VideoParams.ImageFileName, VideoParams.Format });
 	const FString ImageFilePath = FPaths::Combine(DestinationDirectory, ImageFileName);
@@ -306,7 +306,7 @@ FCaptureConvertVideoDataThirdParty::FResult FCaptureConvertVideoDataThirdParty::
 		Task.Update(InProgress);
 	}));
 
-	FString DestinationDirectory = OutputDirectory / Video.UserId;
+	FString DestinationDirectory = OutputDirectory / Video.Name;
 
 	FCaptureConvertVideoDataThirdParty::FResult Result = MakeValue();
 
diff --git a/Engine/Plugins/VirtualProduction/CaptureManager/CaptureManagerApp/Source/CaptureManagerPipeline/Private/Nodes/ConvertAudioNode.cpp b/Engine/Plugins/VirtualProduction/CaptureManager/CaptureManagerApp/Source/CaptureManagerPipeline/Private/Nodes/ConvertAudioNode.cpp
--- a/Engine/Plugins/VirtualProduction/CaptureManager/CaptureManagerApp/Source/CaptureManagerPipeline/Private/Nodes/ConvertAudioNode.cpp
+++ b/Engine/Plugins/VirtualProduction/CaptureManager/CaptureManagerApp/Source/CaptureManagerPipeline/Private/Nodes/ConvertAudioNode.cpp
@@ -61,7 +61,7 @@ FConvertAudioNode::FResult FConvertAudioNode::Validate()
 
 FString FConvertAudioNode::GetAudioDirectory() const
 {
-	return OutputDirectory / Audio.UserId;
+	return OutputDirectory / Audio.Name;
 }
 
 FConvertAudioNode::FResult FConvertAudioNode::CheckForAudioFile(const FString& InAudioDirectory)
diff --git a/Engine/Plugins/VirtualProduction/CaptureManager/CaptureManagerApp/Source/CaptureManagerPipeline/Private/Nodes/ConvertCalibrationNode.cpp b/Engine/Plugins/VirtualProduction/CaptureManager/CaptureManagerApp/Source/CaptureManagerPipeline/Private/Nodes/ConvertCalibrationNode.cpp
--- a/Engine/Plugins/VirtualProduction/CaptureManager/CaptureManagerApp/Source/CaptureManagerPipeline/Private/Nodes/ConvertCalibrationNode.cpp
+++ b/Engine/Plugins/VirtualProduction/CaptureManager/CaptureManagerApp/Source/CaptureManagerPipeline/Private/Nodes/ConvertCalibrationNode.cpp
@@ -92,7 +92,7 @@ FConvertCalibrationNode::FResult FConvertCalibrationNode::Validate()
 
 FString FConvertCalibrationNode::GetCalibrationDirectory() const
 {
-	return OutputDirectory / Calibration.UserId;
+	return OutputDirectory / Calibration.Name;
 }
 
 #undef LOCTEXT_NAMESPACE
\ No newline at end of file
diff --git a/Engine/Plugins/VirtualProduction/CaptureManager/CaptureManagerApp/Source/CaptureManagerPipeline/Private/Nodes/ConvertDepthNode.cpp b/Engine/Plugins/VirtualProduction/CaptureManager/CaptureManagerApp/Source/CaptureManagerPipeline/Private/Nodes/ConvertDepthNode.cpp
--- a/Engine/Plugins/VirtualProduction/CaptureManager/CaptureManagerApp/Source/CaptureManagerPipeline/Private/Nodes/ConvertDepthNode.cpp
+++ b/Engine/Plugins/VirtualProduction/CaptureManager/CaptureManagerApp/Source/CaptureManagerPipeline/Private/Nodes/ConvertDepthNode.cpp
@@ -58,7 +58,7 @@ FConvertDepthNode::FResult FConvertDepthNode::Validate()
 
 FString FConvertDepthNode::GetDepthDirectory() const
 {
-	return OutputDirectory / Depth.UserId;
+	return OutputDirectory / Depth.Name;
 }
 
 FConvertDepthNode::FResult FConvertDepthNode::CheckImagesForDepth(const FString& InDepthPath)
diff --git a/Engine/Plugins/VirtualProduction/CaptureManager/CaptureManagerApp/Source/CaptureManagerPipeline/Private/Nodes/ConvertVideoNode.cpp b/Engine/Plugins/VirtualProduction/CaptureManager/CaptureManagerApp/Source/CaptureManagerPipeline/Private/Nodes/ConvertVideoNode.cpp
--- a/Engine/Plugins/VirtualProduction/CaptureManager/CaptureManagerApp/Source/CaptureManagerPipeline/Private/Nodes/ConvertVideoNode.cpp
+++ b/Engine/Plugins/VirtualProduction/CaptureManager/CaptureManagerApp/Source/CaptureManagerPipeline/Private/Nodes/ConvertVideoNode.cpp
@@ -58,7 +58,7 @@ FConvertVideoNode::FResult FConvertVideoNode::Validate()
 
 FString FConvertVideoNode::GetVideoDirectory() const
 {
-	return OutputDirectory / Video.UserId;
+	return OutputDirectory / Video.Name;
 }
 
 FConvertVideoNode::FResult FConvertVideoNode::CheckImagesForVideo(const FString& InVideoPath)
diff --git a/Engine/Plugins/VirtualProduction/CaptureManager/CaptureManagerApp/Source/StereoCameraMetadata/Private/StereoCameraMetadataParseUtils.cpp b/Engine/Plugins/VirtualProduction/CaptureManager/CaptureManagerApp/Source/StereoCameraMetadata/Private/StereoCameraMetadataParseUtils.cpp
--- a/Engine/Plugins/VirtualProduction/CaptureManager/CaptureManagerApp/Source/StereoCameraMetadata/Private/StereoCameraMetadataParseUtils.cpp
+++ b/Engine/Plugins/VirtualProduction/CaptureManager/CaptureManagerApp/Source/StereoCameraMetadata/Private/StereoCameraMetadataParseUtils.cpp
@@ -71,7 +71,7 @@ FTakeMetadata ConvertOldToNewTakeMetadata(const FString& InTakeFolder, const FSt
 	FString CalibrationPath = InTakeFolder / TEXT("calib.json");
 
 	FTakeMetadata::FCalibration Calibration;
-	Calibration.UserId = TEXT("undefined");
+	Calibration.Name = TEXT("undefined");
 	Calibration.Format = TEXT("opencv");
 	Calibration.Path = InStereoCameraInfo.CalibrationFilePath;
 
@@ -84,7 +84,7 @@ FTakeMetadata ConvertOldToNewTakeMetadata(const FString& InTakeFolder, const FSt
 	Metadata.Calibration.Push(MoveTemp(Calibration));
 	Metadata.Device.Type = InStereoCameraInfo.DeviceInfo.Type;
 	Metadata.Device.Model = InStereoCameraInfo.DeviceInfo.Model;
-	Metadata.Device.UserId = InStereoCameraInfo.DeviceInfo.Id;
+	Metadata.Device.Name = InStereoCameraInfo.DeviceInfo.Id;
 
 	Metadata.UniqueId = InStereoCameraInfo.Id;
 	Metadata.Slate = InStereoCameraInfo.Slate;
@@ -102,7 +102,7 @@ FTakeMetadata ConvertOldToNewTakeMetadata(const FString& InTakeFolder, const FSt
 	for (const TPair<FString, FStereoCameraTakeInfo::FCamera>& Camera : InStereoCameraInfo.CameraMap)
 	{
 		FTakeMetadata::FVideo Video;
-		Video.UserId = Camera.Value.UserId;
+		Video.Name = Camera.Value.UserId;
 		Video.Path = Camera.Value.FramesPath;
 		Video.PathType = FTakeMetadata::FVideo::EPathType::Folder;
 		Video.Format = DetermineImageFormat(Camera.Value.FramesPath);
diff --git a/Engine/Plugins/VirtualProduction/CaptureManager/CaptureManagerApp/Source/StereoCameraMetadata/Private/StereoCameraTakeMetadata.cpp b/Engine/Plugins/VirtualProduction/CaptureManager/CaptureManagerApp/Source/StereoCameraMetadata/Private/StereoCameraTakeMetadata.cpp
--- a/Engine/Plugins/VirtualProduction/CaptureManager/CaptureManagerApp/Source/StereoCameraMetadata/Private/StereoCameraTakeMetadata.cpp
+++ b/Engine/Plugins/VirtualProduction/CaptureManager/CaptureManagerApp/Source/StereoCameraMetadata/Private/StereoCameraTakeMetadata.cpp
@@ -192,7 +192,7 @@ static void CheckCameraIds(const FTakeMetadata& InStereoCameraTakeInfo, const TM
 {
 	for (const FTakeMetadata::FVideo& ExpectedCamera : InStereoCameraTakeInfo.Video)
 	{
-		const FString& ExpectedId = ExpectedCamera.UserId;
+		const FString& ExpectedId = ExpectedCamera.Name;
 
 		if (!OutTakeCameras.Contains(ExpectedId))
 		{
@@ -664,7 +664,7 @@ TOptional<FText> CheckResolutions(const FTakeMetadata::FVideo& InCameraTakeInfo,
 	}
 	else
 	{
-		UE_LOG(LogStereoCameraMetadata, Warning, TEXT("Could not check resolutions as FrameWidth and FrameHeight have not been set for %s at %s"), *InCameraTakeInfo.UserId, *InCameraTakeInfo.Path);
+		UE_LOG(LogStereoCameraMetadata, Warning, TEXT("Could not check resolutions as FrameWidth and FrameHeight have not been set for %s at %s"), *InCameraTakeInfo.Name, *InCameraTakeInfo.Path);
 	}
 
 	return {};
diff --git a/Engine/Plugins/VirtualProduction/CaptureManager/CaptureManagerCore/Source/CaptureManagerTakeMetadata/Private/CaptureManagerTakeMetadata.cpp b/Engine/Plugins/VirtualProduction/CaptureManager/CaptureManagerCore/Source/CaptureManagerTakeMetadata/Private/CaptureManagerTakeMetadata.cpp
--- a/Engine/Plugins/VirtualProduction/CaptureManager/CaptureManagerCore/Source/CaptureManagerTakeMetadata/Private/CaptureManagerTakeMetadata.cpp
+++ b/Engine/Plugins/VirtualProduction/CaptureManager/CaptureManagerCore/Source/CaptureManagerTakeMetadata/Private/CaptureManagerTakeMetadata.cpp
@@ -80,7 +80,14 @@ static FTakeMetadata::FDevice ParseDevice(const rapidjson::Value& InDeviceObject
 	FTakeMetadata::FDevice TakeDevice;
 	TakeDevice.Type = InDeviceObject["Type"].GetString();
 	TakeDevice.Model = InDeviceObject["Model"].GetString();
-	TakeDevice.UserId = InDeviceObject["UserId"].GetString();
+	if (InVersion.Major >= 4 && InVersion.Minor >= 2)
+	{
+		TakeDevice.Name = InDeviceObject["Name"].GetString();
+	}
+	else
+	{
+		TakeDevice.Name = InDeviceObject["UserId"].GetString();
+	}
 
 	if (InDeviceObject.HasMember("Platform"))
 	{
@@ -123,10 +130,18 @@ static FTakeMetadata::FDevice ParseDevice(const rapidjson::Value& InDeviceObject
 	return TakeDevice;
 }
 
-static FTakeMetadata::FVideo ParseVideo(const rapidjson::Value& InVideoObject)
+static FTakeMetadata::FVideo ParseVideo(const rapidjson::Value& InVideoObject, const FTakeMetadata::FSchemaVersion& InVersion)
 {
 	FTakeMetadata::FVideo Video;
-	Video.UserId = InVideoObject["UserId"].GetString();
+
+	if (InVersion.Major >= 4 && InVersion.Minor >= 2)
+	{
+		Video.Name = InVideoObject["Name"].GetString();
+	}
+	else
+	{
+		Video.Name = InVideoObject["UserId"].GetString();
+	}
 	Video.Path = InVideoObject["Path"].GetString();
 
 	if (InVideoObject.HasMember("PathType"))
@@ -206,7 +221,7 @@ static TArray<FTakeMetadata::FCalibration> ParseCalibration(const rapidjson::Val
 	if (InVersion == FTakeMetadata::FSchemaVersion{ 1, 0 })
 	{
 		FTakeMetadata::FCalibration Calibration;
-		Calibration.UserId = "undefined";
+		Calibration.Name = "undefined";
 		Calibration.Path = InCalibrationArray["Path"].GetString();
 		Calibrations.Push(Calibration);
 	}
@@ -215,7 +230,7 @@ static TArray<FTakeMetadata::FCalibration> ParseCalibration(const rapidjson::Val
 		for (const rapidjson::Value& CalibrationObject : InCalibrationArray.GetArray())
 		{
 			FTakeMetadata::FCalibration Calibration;
-			Calibration.UserId = CalibrationObject["UserId"].GetString();
+			Calibration.Name = CalibrationObject["UserId"].GetString();
 			Calibration.Path = CalibrationObject["Path"].GetString();
 			Calibrations.Push(Calibration);
 		}
@@ -225,7 +240,14 @@ static TArray<FTakeMetadata::FCalibration> ParseCalibration(const rapidjson::Val
 		for (const rapidjson::Value& CalibrationObject : InCalibrationArray.GetArray())
 		{
 			FTakeMetadata::FCalibration Calibration;
-			Calibration.UserId = CalibrationObject["UserId"].GetString();
+			if (InVersion.Major >= 4 && InVersion.Minor >= 2)
+			{
+				Calibration.Name = CalibrationObject["Name"].GetString();
+			}
+			else
+			{
+				Calibration.Name = CalibrationObject["UserId"].GetString();
+			}
 			Calibration.Format = CalibrationObject["Format"].GetString();
 			Calibration.Path = CalibrationObject["Path"].GetString();
 			Calibrations.Push(Calibration);
@@ -235,13 +257,21 @@ static TArray<FTakeMetadata::FCalibration> ParseCalibration(const rapidjson::Val
 	return Calibrations;
 }
 
-static TArray<FTakeMetadata::FAudio> ParseAudio(const rapidjson::Value& InAudioArray)
+static TArray<FTakeMetadata::FAudio> ParseAudio(const rapidjson::Value& InAudioArray, const FTakeMetadata::FSchemaVersion& InVersion)
 {
 	TArray<FTakeMetadata::FAudio> Audios;
 	for (const rapidjson::Value& AudioObject : InAudioArray.GetArray())
 	{
 		FTakeMetadata::FAudio Audio;
-		Audio.UserId = AudioObject["UserId"].GetString();
+		if (InVersion.Major >= 4 && InVersion.Minor >= 2)
+		{
+			Audio.Name = AudioObject["Name"].GetString();
+		}
+		else
+		{
+			Audio.Name = AudioObject["UserId"].GetString();
+		}
+
 		Audio.Path = AudioObject["Path"].GetString();
 		if (AudioObject.HasMember("Duration"))
 		{
@@ -289,12 +319,12 @@ static FTakeMetadata ParseTakeMetadata(const rapidjson::Document& InDocument, co
 
 	for (const rapidjson::Value& VideoObject : InDocument["Video"].GetArray())
 	{
-		TakeMetadata.Video.Push(ParseVideo(VideoObject));
+		TakeMetadata.Video.Push(ParseVideo(VideoObject, Version));
 	}
 
 	for (const rapidjson::Value& DepthObject : InDocument["Depth"].GetArray())
 	{
-		TakeMetadata.Depth.Push(ParseVideo(DepthObject));
+		TakeMetadata.Depth.Push(ParseVideo(DepthObject, Version));
 	}
 
 	if (InDocument.HasMember("Calibration"))
@@ -304,7 +334,7 @@ static FTakeMetadata ParseTakeMetadata(const rapidjson::Document& InDocument, co
 
 	if (InDocument.HasMember("Audio"))
 	{
-		TakeMetadata.Audio = ParseAudio(InDocument["Audio"]);
+		TakeMetadata.Audio = ParseAudio(InDocument["Audio"], Version);
 	}
 
 	return TakeMetadata;
@@ -336,7 +366,7 @@ rapidjson::Value SerializeDevice(const FTakeMetadata::FDevice& InDevice, rapidjs
 
 	DeviceObject.AddMember("Type", SerializeString(InDevice.Type, InAllocator), InAllocator);
 	DeviceObject.AddMember("Model", SerializeString(InDevice.Model, InAllocator), InAllocator);
-	DeviceObject.AddMember("UserId", SerializeString(InDevice.UserId, InAllocator), InAllocator);
+	DeviceObject.AddMember("Name", SerializeString(InDevice.Name, InAllocator), InAllocator);
 
 	if (InDevice.Platform.IsSet())
 	{
@@ -408,7 +438,7 @@ rapidjson::Value SerializeVideo(const TArray<FTakeMetadata::FVideo>& InVideos, r
 	for (const FTakeMetadata::FVideo& Video : InVideos)
 	{
 		rapidjson::Value VideoObject(rapidjson::kObjectType);
-		VideoObject.AddMember("UserId", SerializeString(Video.UserId, InAllocator), InAllocator);
+		VideoObject.AddMember("Name", SerializeString(Video.Name, InAllocator), InAllocator);
 		VideoObject.AddMember("Path", SerializeString(Video.Path, InAllocator), InAllocator);
 
 		if (Video.PathType.IsSet())
@@ -467,7 +497,7 @@ rapidjson::Value SerializeCalibration(const TArray <FTakeMetadata::FCalibration>
 	for (const FTakeMetadata::FCalibration& Calibration : InCalibration)
 	{
 		rapidjson::Value CalibrationObject(rapidjson::kObjectType);
-		CalibrationObject.AddMember("UserId", SerializeString(Calibration.UserId, InAllocator), InAllocator);
+		CalibrationObject.AddMember("Name", SerializeString(Calibration.Name, InAllocator), InAllocator);
 		CalibrationObject.AddMember("Format", SerializeString(Calibration.Format, InAllocator), InAllocator);
 		CalibrationObject.AddMember("Path", SerializeString(Calibration.Path, InAllocator), InAllocator);
 
@@ -484,7 +514,7 @@ rapidjson::Value SerializeAudio(const TArray<FTakeMetadata::FAudio>& InAudio, ra
 	for (const FTakeMetadata::FAudio& Audio : InAudio)
 	{
 		rapidjson::Value AudioObject(rapidjson::kObjectType);
-		AudioObject.AddMember("UserId", SerializeString(Audio.UserId, InAllocator), InAllocator);
+		AudioObject.AddMember("Name", SerializeString(Audio.Name, InAllocator), InAllocator);
 		AudioObject.AddMember("Path", SerializeString(Audio.Path, InAllocator), InAllocator);
 		
 		if (Audio.Duration.IsSet())
diff --git a/Engine/Plugins/VirtualProduction/CaptureManager/CaptureManagerCore/Source/CaptureManagerTakeMetadata/Private/Tests/CaptureManagerTakeMetadata.spec.cpp b/Engine/Plugins/VirtualProduction/CaptureManager/CaptureManagerCore/Source/CaptureManagerTakeMetadata/Private/Tests/CaptureManagerTakeMetadata.spec.cpp
--- a/Engine/Plugins/VirtualProduction/CaptureManager/CaptureManagerCore/Source/CaptureManagerTakeMetadata/Private/Tests/CaptureManagerTakeMetadata.spec.cpp
+++ b/Engine/Plugins/VirtualProduction/CaptureManager/CaptureManagerCore/Source/CaptureManagerTakeMetadata/Private/Tests/CaptureManagerTakeMetadata.spec.cpp
@@ -42,10 +42,48 @@ void CaptureManagerTakeMetadata_Spec::TestTakeMetadataSchemaSpec::Define()
 
 	Describe(TEXT("TestTakeMetadataSchemaTests"), [this]
 	{
+			It(TEXT("take_valid_v4_2_mandatory_media"), [this]
+				{
+					TValueOrError<FTakeMetadata, FTakeMetadataParserError> Result =
+						GetTakeMetadataParser().Parse(GetJsonTestsDir() / TEXT("take_valid_v4_2_mandatory_media.cptake"));
+
+					UTEST_TRUE(TEXT("Should have value."), Result.HasValue());
+
+					FTakeMetadata TakeMetadata = Result.GetValue();
+
+					TestEqual(TEXT("Version.Major has unexpected value."), TakeMetadata.Version.Major, 4);
+					TestEqual(TEXT("Version.Minor has unexpected value."), TakeMetadata.Version.Minor, 2);
+
+					TestEqual(TEXT("UniqueId has unexpected value."), TakeMetadata.UniqueId, TEXT("a78613f3-e660-47e4-af6a-1298cde7c947"));
+					TestEqual(TEXT("TakeNumber has unexpected value."), TakeMetadata.TakeNumber, 1);
+					TestEqual(TEXT("Slate has unexpected value."), TakeMetadata.Slate, TEXT("Fox Jumps Over The Lazy Dog"));
+
+					TestEqual(TEXT("Device.Type has unexpected value."), TakeMetadata.Device.Type, TEXT("HMC"));
+					TestEqual(TEXT("Device.Model has unexpected value."), TakeMetadata.Device.Model, TEXT("StereoHMC"));
+					TestEqual(TEXT("Device.Name has unexpected value."), TakeMetadata.Device.Name, TEXT("UserDev001"));
+
+					UTEST_EQUAL(TEXT("Video array is expected to have one element."), TakeMetadata.Video.Num(), 1);
+					TestEqual(TEXT("Video[0].Name has unexpected value."), TakeMetadata.Video[0].Name, TEXT("secondary"));
+					TestEqual(TEXT("Video[0].Path has unexpected value."), TakeMetadata.Video[0].Path, TEXT("folder_or_file_name"));
+					TestEqual(TEXT("Video[0].Format has unexpected value."), TakeMetadata.Video[0].Format, TEXT("mov"));
+					TestEqual(TEXT("Video[0].FrameRate has unexpected value."), TakeMetadata.Video[0].FrameRate, 60.0f);
+
+					UTEST_EQUAL(TEXT("Calibration array is expected to have one element."), TakeMetadata.Calibration.Num(), 1);
+					TestEqual(TEXT("Calibration[0].Name has unexpected value."), TakeMetadata.Calibration[0].Name, TEXT("calibration_user_id"));
+					TestEqual(TEXT("Calibration[0].Format has unexpected value."), TakeMetadata.Calibration[0].Format, TEXT("opencv"));
+					TestEqual(TEXT("Calibration[0].Path has unexpected value."), TakeMetadata.Calibration[0].Path, TEXT("calib.json"));
+
+					UTEST_EQUAL(TEXT("Audio array is expected to have one element."), TakeMetadata.Audio.Num(), 1);
+					TestEqual(TEXT("Audio[0].Name has unexpected value."), TakeMetadata.Audio[0].Name, TEXT("primary"));
+					TestEqual(TEXT("Audio[0].Path has unexpected value."), TakeMetadata.Audio[0].Path, TEXT("audio.wav"));
+
+					return true;
+				});
+
 			It(TEXT("take_valid_v4_0_mandatory"), [this]
 				{
 					TValueOrError<FTakeMetadata, FTakeMetadataParserError> Result =
-						GetTakeMetadataParser().Parse(GetJsonTestsDir() / TEXT("take_valid_v4_0_mandatory.json"));
+						GetTakeMetadataParser().Parse(GetJsonTestsDir() / TEXT("take_valid_v4_0_mandatory.cptake"));
 
 					UTEST_TRUE(TEXT("Should have value."), Result.HasValue());
 
@@ -60,7 +98,7 @@ void CaptureManagerTakeMetadata_Spec::TestTakeMetadataSchemaSpec::Define()
 
 					TestEqual(TEXT("Device.Type has unexpected value."), TakeMetadata.Device.Type, TEXT("HMC"));
 					TestEqual(TEXT("Device.Model has unexpected value."), TakeMetadata.Device.Model, TEXT("StereoHMC"));
-					TestEqual(TEXT("Device.UserId has unexpected value."), TakeMetadata.Device.UserId, TEXT("UserDev001"));
+					TestEqual(TEXT("Device.Name has unexpected value."), TakeMetadata.Device.Name, TEXT("UserDev001"));
 
 					return true;
 				});
@@ -68,7 +106,7 @@ void CaptureManagerTakeMetadata_Spec::TestTakeMetadataSchemaSpec::Define()
 			It(TEXT("take_valid_v4_0_mandatory_media"), [this]
 				{
 					TValueOrError<FTakeMetadata, FTakeMetadataParserError> Result =
-						GetTakeMetadataParser().Parse(GetJsonTestsDir() / TEXT("take_valid_v4_0_mandatory_media.json"));
+						GetTakeMetadataParser().Parse(GetJsonTestsDir() / TEXT("take_valid_v4_0_mandatory_media.cptake"));
 
 					UTEST_TRUE(TEXT("Should have value."), Result.HasValue());
 
@@ -83,21 +121,21 @@ void CaptureManagerTakeMetadata_Spec::TestTakeMetadataSchemaSpec::Define()
 
 					TestEqual(TEXT("Device.Type has unexpected value."), TakeMetadata.Device.Type, TEXT("HMC"));
 					TestEqual(TEXT("Device.Model has unexpected value."), TakeMetadata.Device.Model, TEXT("StereoHMC"));
-					TestEqual(TEXT("Device.UserId has unexpected value."), TakeMetadata.Device.UserId, TEXT("UserDev001"));
+					TestEqual(TEXT("Device.Name has unexpected value."), TakeMetadata.Device.Name, TEXT("UserDev001"));
 
 					UTEST_EQUAL(TEXT("Video array is expected to have one element."), TakeMetadata.Video.Num(), 1);
-					TestEqual(TEXT("Video[0].UserId has unexpected value."), TakeMetadata.Video[0].UserId, TEXT("secondary"));
+					TestEqual(TEXT("Video[0].Name has unexpected value."), TakeMetadata.Video[0].Name, TEXT("secondary"));
 					TestEqual(TEXT("Video[0].Path has unexpected value."), TakeMetadata.Video[0].Path, TEXT("folder_or_file_name"));
 					TestEqual(TEXT("Video[0].Format has unexpected value."), TakeMetadata.Video[0].Format, TEXT("mov"));
 					TestEqual(TEXT("Video[0].FrameRate has unexpected value."), TakeMetadata.Video[0].FrameRate, 60.0f);
 
 					UTEST_EQUAL(TEXT("Calibration array is expected to have one element."), TakeMetadata.Calibration.Num(), 1);
-					TestEqual(TEXT("Calibration[0].UserId has unexpected value."), TakeMetadata.Calibration[0].UserId, TEXT("calibration_user_id"));
+					TestEqual(TEXT("Calibration[0].Name has unexpected value."), TakeMetadata.Calibration[0].Name, TEXT("calibration_user_id"));
 					TestEqual(TEXT("Calibration[0].Format has unexpected value."), TakeMetadata.Calibration[0].Format, TEXT("opencv"));
 					TestEqual(TEXT("Calibration[0].Path has unexpected value."), TakeMetadata.Calibration[0].Path, TEXT("calib.json"));
 
 					UTEST_EQUAL(TEXT("Audio array is expected to have one element."), TakeMetadata.Audio.Num(), 1);
-					TestEqual(TEXT("Audio[0].UserId has unexpected value."), TakeMetadata.Audio[0].UserId, TEXT("primary"));
+					TestEqual(TEXT("Audio[0].Name has unexpected value."), TakeMetadata.Audio[0].Name, TEXT("primary"));
 					TestEqual(TEXT("Audio[0].Path has unexpected value."), TakeMetadata.Audio[0].Path, TEXT("audio.wav"));
 
 					return true;
@@ -106,7 +144,7 @@ void CaptureManagerTakeMetadata_Spec::TestTakeMetadataSchemaSpec::Define()
 			It(TEXT("take_valid_v4_0"), [this]
 				{
 					TValueOrError<FTakeMetadata, FTakeMetadataParserError> Result =
-						GetTakeMetadataParser().Parse(GetJsonTestsDir() / TEXT("take_valid_v4_0.json"));
+						GetTakeMetadataParser().Parse(GetJsonTestsDir() / TEXT("take_valid_v4_0.cptake"));
 
 					UTEST_TRUE(TEXT("Should have value."), Result.HasValue());
 
@@ -126,7 +164,7 @@ void CaptureManagerTakeMetadata_Spec::TestTakeMetadataSchemaSpec::Define()
 
 					TestEqual(TEXT("Device.Type has unexpected value."), TakeMetadata.Device.Type, TEXT("HMC"));
 					TestEqual(TEXT("Device.Model has unexpected value."), TakeMetadata.Device.Model, TEXT("StereoHMC"));
-					TestEqual(TEXT("Device.UserId has unexpected value."), TakeMetadata.Device.UserId, TEXT("UserDev001"));
+					TestEqual(TEXT("Device.Name has unexpected value."), TakeMetadata.Device.Name, TEXT("UserDev001"));
 					TestEqual(TEXT("Device.Platform.Name has unexpected value."), TakeMetadata.Device.Platform.GetValue().Name, TEXT("iOS"));
 					TestEqual(TEXT("Device.Platform.Version.Major has unexpected value."), TakeMetadata.Device.Platform.GetValue().Version.GetValue(), "17.2.1");
 
@@ -135,7 +173,7 @@ void CaptureManagerTakeMetadata_Spec::TestTakeMetadataSchemaSpec::Define()
 					TestEqual(TEXT("Device.Software.Version.Major has unexpected value."), TakeMetadata.Device.Software[0].Version.GetValue(), "2.5.7");
 
 					UTEST_EQUAL(TEXT("Video array is expected to have one element."), TakeMetadata.Video.Num(), 1);
-					TestEqual(TEXT("Video[0].UserId has unexpected value."), TakeMetadata.Video[0].UserId, TEXT("secondary"));
+					TestEqual(TEXT("Video[0].Name has unexpected value."), TakeMetadata.Video[0].Name, TEXT("secondary"));
 					TestEqual(TEXT("Video[0].Path has unexpected value."), TakeMetadata.Video[0].Path, TEXT("folder_or_file_name"));
 					TestEqual(TEXT("Video[0].PathType has unexpected value."), TakeMetadata.Video[0].PathType.GetValue(), FTakeMetadata::FVideo::EPathType::Folder);
 					TestEqual(TEXT("Video[0].Format has unexpected value."), TakeMetadata.Video[0].Format, TEXT("mov"));
@@ -155,12 +193,12 @@ void CaptureManagerTakeMetadata_Spec::TestTakeMetadataSchemaSpec::Define()
 					TestEqual(TEXT("Video[0].TimecodeStart has unexpected value."), TakeMetadata.Video[0].TimecodeStart.GetValue(), TEXT("09:01:41:01.300"));
 
 					UTEST_EQUAL(TEXT("Calibration array is expected to have one element."), TakeMetadata.Calibration.Num(), 1);
-					TestEqual(TEXT("Calibration[0].UserId has unexpected value."), TakeMetadata.Calibration[0].UserId, TEXT("calibration_user_id"));
+					TestEqual(TEXT("Calibration[0].Name has unexpected value."), TakeMetadata.Calibration[0].Name, TEXT("calibration_user_id"));
 					TestEqual(TEXT("Calibration[0].Format has unexpected value."), TakeMetadata.Calibration[0].Format, TEXT("opencv"));
 					TestEqual(TEXT("Calibration[0].Path has unexpected value."), TakeMetadata.Calibration[0].Path, TEXT("calib.json"));
 
 					UTEST_EQUAL(TEXT("Audio array is expected to have one element."), TakeMetadata.Audio.Num(), 1);
-					TestEqual(TEXT("Audio[0].UserId has unexpected value."), TakeMetadata.Audio[0].UserId, TEXT("primary"));
+					TestEqual(TEXT("Audio[0].Name has unexpected value."), TakeMetadata.Audio[0].Name, TEXT("primary"));
 					TestEqual(TEXT("Audio[0].Path has unexpected value."), TakeMetadata.Audio[0].Path, TEXT("audio.wav"));
 					TestEqual(TEXT("Audio[0].Duration has unexpected value."), TakeMetadata.Audio[0].Duration.GetValue(), 420.0f);
 					TestEqual(TEXT("Audio[0].TimecodeRate has unexpected value."), TakeMetadata.Audio[0].TimecodeRate.GetValue(), 60.0f);
@@ -172,7 +210,7 @@ void CaptureManagerTakeMetadata_Spec::TestTakeMetadataSchemaSpec::Define()
 			It(TEXT("take_valid_v3_0"), [this]
 				{
 					TValueOrError<FTakeMetadata, FTakeMetadataParserError> Result =
-						GetTakeMetadataParser().Parse(GetJsonTestsDir() / TEXT("take_valid_v3_0.json"));
+						GetTakeMetadataParser().Parse(GetJsonTestsDir() / TEXT("take_valid_v3_0.cptake"));
 
 					UTEST_TRUE(TEXT("Should have value."), Result.HasValue());
 
@@ -192,7 +230,7 @@ void CaptureManagerTakeMetadata_Spec::TestTakeMetadataSchemaSpec::Define()
 
 					TestEqual(TEXT("Device.Type has unexpected value."), TakeMetadata.Device.Type, TEXT("HMC"));
 					TestEqual(TEXT("Device.Model has unexpected value."), TakeMetadata.Device.Model, TEXT("StereoHMC"));
-					TestEqual(TEXT("Device.UserId has unexpected value."), TakeMetadata.Device.UserId, TEXT("UserDev001"));
+					TestEqual(TEXT("Device.Name has unexpected value."), TakeMetadata.Device.Name, TEXT("UserDev001"));
 					TestEqual(TEXT("Device.Platform.Name has unexpected value."), TakeMetadata.Device.Platform.GetValue().Name, TEXT("iOS"));
 					TestEqual(TEXT("Device.Platform.Version.Major has unexpected value."), TakeMetadata.Device.Platform.GetValue().Version.GetValue(), "17.2.1");
 
@@ -201,7 +239,7 @@ void CaptureManagerTakeMetadata_Spec::TestTakeMetadataSchemaSpec::Define()
 					TestEqual(TEXT("Device.Software.Version.Major has unexpected value."), TakeMetadata.Device.Software[0].Version.GetValue(), "2.5.7");
 
 					UTEST_EQUAL(TEXT("Video array is expected to have one element."), TakeMetadata.Video.Num(), 1);
-					TestEqual(TEXT("Video[0].UserId has unexpected value."), TakeMetadata.Video[0].UserId, TEXT("secondary"));
+					TestEqual(TEXT("Video[0].Name has unexpected value."), TakeMetadata.Video[0].Name, TEXT("secondary"));
 					TestEqual(TEXT("Video[0].Path has unexpected value."), TakeMetadata.Video[0].Path, TEXT("folder_or_file_name"));
 					TestEqual(TEXT("Video[0].PathType has unexpected value."), TakeMetadata.Video[0].PathType.GetValue(), FTakeMetadata::FVideo::EPathType::Folder);
 					TestEqual(TEXT("Video[0].Format has unexpected value."), TakeMetadata.Video[0].Format, TEXT("mov"));
@@ -221,12 +259,12 @@ void CaptureManagerTakeMetadata_Spec::TestTakeMetadataSchemaSpec::Define()
 					TestEqual(TEXT("Video[0].TimecodeStart has unexpected value."), TakeMetadata.Video[0].TimecodeStart.GetValue(), TEXT("09:01:41:01.300"));
 
 					UTEST_EQUAL(TEXT("Calibration array is expected to have one element."), TakeMetadata.Calibration.Num(), 1);
-					TestEqual(TEXT("Calibration[0].UserId has unexpected value."), TakeMetadata.Calibration[0].UserId, TEXT("calibration_user_id"));
+					TestEqual(TEXT("Calibration[0].Name has unexpected value."), TakeMetadata.Calibration[0].Name, TEXT("calibration_user_id"));
 					TestEqual(TEXT("Calibration[0].Format has unexpected value."), TakeMetadata.Calibration[0].Format, TEXT("opencv"));
 					TestEqual(TEXT("Calibration[0].Path has unexpected value."), TakeMetadata.Calibration[0].Path, TEXT("calib.json"));
 
 					UTEST_EQUAL(TEXT("Audio array is expected to have one element."), TakeMetadata.Audio.Num(), 1);
-					TestEqual(TEXT("Audio[0].UserId has unexpected value."), TakeMetadata.Audio[0].UserId, TEXT("primary"));
+					TestEqual(TEXT("Audio[0].Name has unexpected value."), TakeMetadata.Audio[0].Name, TEXT("primary"));
 					TestEqual(TEXT("Audio[0].Path has unexpected value."), TakeMetadata.Audio[0].Path, TEXT("audio.wav"));
 					TestEqual(TEXT("Audio[0].Duration has unexpected value."), TakeMetadata.Audio[0].Duration.GetValue(), 420.0f);
 					TestEqual(TEXT("Audio[0].TimecodeRate has unexpected value."), TakeMetadata.Audio[0].TimecodeRate.GetValue(), 60.0f);
@@ -238,7 +276,7 @@ void CaptureManagerTakeMetadata_Spec::TestTakeMetadataSchemaSpec::Define()
 		It(TEXT("take_valid_v2_0"), [this]
 			{
 				TValueOrError<FTakeMetadata, FTakeMetadataParserError> Result =
-					GetTakeMetadataParser().Parse(GetJsonTestsDir() / TEXT("take_valid_v2_0.json"));
+					GetTakeMetadataParser().Parse(GetJsonTestsDir() / TEXT("take_valid_v2_0.cptake"));
 
 				UTEST_TRUE(TEXT("Should have value."), Result.HasValue());
 
@@ -258,7 +296,7 @@ void CaptureManagerTakeMetadata_Spec::TestTakeMetadataSchemaSpec::Define()
 
 				TestEqual(TEXT("Device.Type has unexpected value."), TakeMetadata.Device.Type, TEXT("HMC"));
 				TestEqual(TEXT("Device.Model has unexpected value."), TakeMetadata.Device.Model, TEXT("StereoHMC"));
-				TestEqual(TEXT("Device.UserId has unexpected value."), TakeMetadata.Device.UserId, TEXT("UserDev001"));
+				TestEqual(TEXT("Device.Name has unexpected value."), TakeMetadata.Device.Name, TEXT("UserDev001"));
 				TestEqual(TEXT("Device.Platform.Name has unexpected value."), TakeMetadata.Device.Platform.GetValue().Name, TEXT("iOS"));
 				TestEqual(TEXT("Device.Platform.Version.Major has unexpected value."), TakeMetadata.Device.Platform.GetValue().Version.GetValue(), "17.2.1");
 					
@@ -267,7 +305,7 @@ void CaptureManagerTakeMetadata_Spec::TestTakeMetadataSchemaSpec::Define()
 				TestEqual(TEXT("Device.Software.Version.Major has unexpected value."), TakeMetadata.Device.Software[0].Version.GetValue(), "2.5.7");
 					
 				UTEST_EQUAL(TEXT("Video array is expected to have one element."), TakeMetadata.Video.Num(), 1);
-				TestEqual(TEXT("Video[0].UserId has unexpected value."), TakeMetadata.Video[0].UserId, TEXT("secondary"));
+				TestEqual(TEXT("Video[0].Name has unexpected value."), TakeMetadata.Video[0].Name, TEXT("secondary"));
 				TestEqual(TEXT("Video[0].Path has unexpected value."), TakeMetadata.Video[0].Path, TEXT("folder_or_file_name"));
 				TestEqual(TEXT("Video[0].PathType has unexpected value."), TakeMetadata.Video[0].PathType.GetValue(), FTakeMetadata::FVideo::EPathType::Folder);
 				TestEqual(TEXT("Video[0].Format has unexpected value."), TakeMetadata.Video[0].Format, TEXT("mov"));
@@ -287,11 +325,11 @@ void CaptureManagerTakeMetadata_Spec::TestTakeMetadataSchemaSpec::Define()
 				TestEqual(TEXT("Video[0].TimecodeStart has unexpected value."), TakeMetadata.Video[0].TimecodeStart.GetValue(), TEXT("09:01:41:01.300"));
 
 				UTEST_EQUAL(TEXT("Calibration array is expected to have one element."), TakeMetadata.Calibration.Num(), 1);
-				TestEqual(TEXT("Calibration[0].UserId has unexpected value."), TakeMetadata.Calibration[0].UserId, TEXT("calibration_user_id"));
+				TestEqual(TEXT("Calibration[0].Name has unexpected value."), TakeMetadata.Calibration[0].Name, TEXT("calibration_user_id"));
 				TestEqual(TEXT("Calibration[0].Path has unexpected value."), TakeMetadata.Calibration[0].Path, TEXT("calib.json"));
 
 				UTEST_EQUAL(TEXT("Audio array is expected to have one element."), TakeMetadata.Audio.Num(), 1);
-				TestEqual(TEXT("Audio[0].UserId has unexpected value."), TakeMetadata.Audio[0].UserId, TEXT("primary"));
+				TestEqual(TEXT("Audio[0].Name has unexpected value."), TakeMetadata.Audio[0].Name, TEXT("primary"));
 				TestEqual(TEXT("Audio[0].Path has unexpected value."), TakeMetadata.Audio[0].Path, TEXT("audio.wav"));
 				TestEqual(TEXT("Audio[0].Duration has unexpected value."), TakeMetadata.Audio[0].Duration.GetValue(), 420.0f);
 				TestEqual(TEXT("Audio[0].TimecodeRate has unexpected value."), TakeMetadata.Audio[0].TimecodeRate.GetValue(), 60.0f);
@@ -302,7 +340,7 @@ void CaptureManagerTakeMetadata_Spec::TestTakeMetadataSchemaSpec::Define()
 
 		It(TEXT("take_valid_v1_0_calibration"), [this]
 			{
-				TValueOrError<FTakeMetadata, FTakeMetadataParserError> Result = GetTakeMetadataParser().Parse(GetJsonTestsDir() / TEXT("take_valid_v1_0_calibration.json"));
+				TValueOrError<FTakeMetadata, FTakeMetadataParserError> Result = GetTakeMetadataParser().Parse(GetJsonTestsDir() / TEXT("take_valid_v1_0_calibration.cptake"));
 				TestTrue(TEXT("Should have value."), Result.HasValue());
 				if (Result.HasValue())
 				{
@@ -310,7 +348,7 @@ void CaptureManagerTakeMetadata_Spec::TestTakeMetadataSchemaSpec::Define()
 					TestEqual(TEXT("Calibration array is expected to have one element."), TakeMetadata.Calibration.Num(), 1);
 					if (TakeMetadata.Calibration.Num() == 1)
 					{
-						TestEqual(TEXT("Calibration[0].UserId has unexpected value."), TakeMetadata.Calibration[0].UserId, TEXT("undefined"));
+						TestEqual(TEXT("Calibration[0].Name has unexpected value."), TakeMetadata.Calibration[0].Name, TEXT("undefined"));
 						TestEqual(TEXT("Calibration[0].Path has unexpected value."), TakeMetadata.Calibration[0].Path, TEXT("calib.json"));
 					}
 				}
@@ -346,7 +384,7 @@ void CaptureManagerTakeMetadata_Spec::TestTakeMetadataSchemaSpec::Define()
 		{
 			It(Test.Name, [this, Test]
 				{
-					TValueOrError<FTakeMetadata, FTakeMetadataParserError> Result = GetTakeMetadataParser().Parse(GetJsonTestsDir() / Test.Name + ".json");
+					TValueOrError<FTakeMetadata, FTakeMetadataParserError> Result = GetTakeMetadataParser().Parse(GetJsonTestsDir() / Test.Name + ".cptake");
 					if (!TestTrue(TEXT("Should have error."), Result.HasError()))
 					{
 						return;
@@ -370,7 +408,7 @@ void CaptureManagerTakeMetadata_Spec::TestTakeMetadataSchemaSpec::Define()
 					{
 						FString TestOutputErrorMessage = TEXT("Unexpected error message. Got: \n\t");
 						TestOutputErrorMessage += *ErrorMessage;
-						TestOutputErrorMessage += "\nbut expected to find thse in order:\n\t";
+						TestOutputErrorMessage += "\nbut expected to find these in order:\n\t";
 						for (FString Keyword : Test.ExpectedMessageKeywords)
 						{
 							TestOutputErrorMessage += "'" + Keyword + "', ";
diff --git a/Engine/Plugins/VirtualProduction/CaptureManager/CaptureManagerCore/Source/CaptureManagerTakeMetadata/Public/CaptureManagerTakeMetadata.h b/Engine/Plugins/VirtualProduction/CaptureManager/CaptureManagerCore/Source/CaptureManagerTakeMetadata/Public/CaptureManagerTakeMetadata.h
--- a/Engine/Plugins/VirtualProduction/CaptureManager/CaptureManagerCore/Source/CaptureManagerTakeMetadata/Public/CaptureManagerTakeMetadata.h
+++ b/Engine/Plugins/VirtualProduction/CaptureManager/CaptureManagerCore/Source/CaptureManagerTakeMetadata/Public/CaptureManagerTakeMetadata.h
@@ -69,7 +69,7 @@ public:
 			TOptional<FString> Version;
 		};
 
-		FString UserId;
+		FString Name;
 		FString Type;
 		FString Model;
 		
@@ -95,7 +95,7 @@ public:
 			CW270
 		};
 
-		FString UserId;
+		FString Name;
 		FString Path;
 
 		/** Path type (e.g. Folder or File). */
@@ -126,7 +126,7 @@ public:
 	/** Calibration information. */
 	struct FCalibration
 	{
-		FString UserId;
+		FString Name;
 		FString Path;
 
 		/** File format (e.g. mhaical, unreal). */
@@ -136,7 +136,7 @@ public:
 	/** Audio information. */
 	struct FAudio
 	{
-		FString UserId;
+		FString Name;
 		FString Path;
 
 		/** Duration in seconds. */
diff --git a/Engine/Plugins/VirtualProduction/CaptureManager/CaptureManagerCore/Source/DataIngestCore/Private/IngestCaptureData.cpp b/Engine/Plugins/VirtualProduction/CaptureManager/CaptureManagerCore/Source/DataIngestCore/Private/IngestCaptureData.cpp
--- a/Engine/Plugins/VirtualProduction/CaptureManager/CaptureManagerCore/Source/DataIngestCore/Private/IngestCaptureData.cpp
+++ b/Engine/Plugins/VirtualProduction/CaptureManager/CaptureManagerCore/Source/DataIngestCore/Private/IngestCaptureData.cpp
@@ -41,8 +41,8 @@ TValueOrError<FIngestCaptureData::FVideo, FText> ParseVideoObject(const TSharedP
 
 	const TSharedPtr<FJsonObject>& ObjectRef = *Object;
 
-	INGEST_CAPTURE_DATA_CHECK_AND_RETURN(ObjectRef->TryGetStringField(TEXT("UserId"), Video.UserId),
-										 LOCTEXT("Parse_MissingVideoUserId", "Video object doesn't contain UserId field"));
+	INGEST_CAPTURE_DATA_CHECK_AND_RETURN(ObjectRef->TryGetStringField(TEXT("Name"), Video.Name),
+										 LOCTEXT("Parse_MissingVideoName", "Video object doesn't contain Name field"));
 
 	INGEST_CAPTURE_DATA_CHECK_AND_RETURN(ObjectRef->TryGetStringField(TEXT("Path"), Video.Path),
 										 LOCTEXT("Parse_MissingVideoPath", "Video object doesn't contain Path field"));
@@ -95,8 +95,8 @@ TValueOrError<FIngestCaptureData::FAudio, FText> ParseAudioObject(const TSharedP
 
 	const TSharedPtr<FJsonObject>& ObjectRef = *Object;
 
-	INGEST_CAPTURE_DATA_CHECK_AND_RETURN(ObjectRef->TryGetStringField(TEXT("UserId"), Audio.UserId),
-										 LOCTEXT("Parse_MissingAudioUserId", "Audio object doesn't contain UserId field"));
+	INGEST_CAPTURE_DATA_CHECK_AND_RETURN(ObjectRef->TryGetStringField(TEXT("Name"), Audio.Name),
+										 LOCTEXT("Parse_MissingAudioName", "Audio object doesn't contain Name field"));
 
 	INGEST_CAPTURE_DATA_CHECK_AND_RETURN(ObjectRef->TryGetStringField(TEXT("Path"), Audio.Path),
 										 LOCTEXT("Parse_MissingAudioPath", "Audio object doesn't contain Path field"));
@@ -125,8 +125,8 @@ TValueOrError<FIngestCaptureData::FCalibration, FText> ParseCalibrationObject(co
 
 	const TSharedPtr<FJsonObject>& ObjectRef = *Object;
 
-	INGEST_CAPTURE_DATA_CHECK_AND_RETURN(ObjectRef->TryGetStringField(TEXT("UserId"), Calibration.UserId),
-										 LOCTEXT("Parse_MissingCalibrationUserId", "Calibration object doesn't contain UserId field"));
+	INGEST_CAPTURE_DATA_CHECK_AND_RETURN(ObjectRef->TryGetStringField(TEXT("Name"), Calibration.Name),
+										 LOCTEXT("Parse_MissingCalibrationName", "Calibration object doesn't contain Name field"));
 
 	INGEST_CAPTURE_DATA_CHECK_AND_RETURN(ObjectRef->TryGetStringField(TEXT("Path"), Calibration.Path),
 										 LOCTEXT("Parse_MissingCalibrationPath", "Calibration object doesn't contain Path field"));
@@ -237,7 +237,7 @@ void SerializeVideo(TSharedPtr<TJsonWriter<TCHAR>> InJsonWriter, const FIngestCa
 {
 	InJsonWriter->WriteObjectStart();
 
-	InJsonWriter->WriteValue(TEXT("UserId"), InVideo.UserId);
+	InJsonWriter->WriteValue(TEXT("Name"), InVideo.Name);
 	InJsonWriter->WriteValue(TEXT("Path"), InVideo.Path);
 
 	if (InVideo.FrameRate.IsSet())
@@ -279,7 +279,7 @@ void SerializeAudio(TSharedPtr<TJsonWriter<TCHAR>> InJsonWriter, const FIngestCa
 {
 	InJsonWriter->WriteObjectStart();
 
-	InJsonWriter->WriteValue(TEXT("UserId"), InAudio.UserId);
+	InJsonWriter->WriteValue(TEXT("Name"), InAudio.Name);
 	InJsonWriter->WriteValue(TEXT("Path"), InAudio.Path);
 
 	if (InAudio.TimecodeStart.IsSet())
@@ -299,7 +299,7 @@ void SerializeCalibration(TSharedPtr<TJsonWriter<TCHAR>> InJsonWriter, const FIn
 {
 	InJsonWriter->WriteObjectStart();
 
-	InJsonWriter->WriteValue(TEXT("UserId"), InCalibration.UserId);
+	InJsonWriter->WriteValue(TEXT("Name"), InCalibration.Name);
 	InJsonWriter->WriteValue(TEXT("Path"), InCalibration.Path);
 
 	InJsonWriter->WriteObjectEnd();
diff --git a/Engine/Plugins/VirtualProduction/CaptureManager/CaptureManagerCore/Source/DataIngestCore/Public/IngestCaptureData.h b/Engine/Plugins/VirtualProduction/CaptureManager/CaptureManagerCore/Source/DataIngestCore/Public/IngestCaptureData.h
--- a/Engine/Plugins/VirtualProduction/CaptureManager/CaptureManagerCore/Source/DataIngestCore/Public/IngestCaptureData.h
+++ b/Engine/Plugins/VirtualProduction/CaptureManager/CaptureManagerCore/Source/DataIngestCore/Public/IngestCaptureData.h
@@ -14,7 +14,7 @@ struct DATAINGESTCORE_API FIngestCaptureData
 	/** Video/image sequence information. */
 	struct FVideo
 	{
-		FString UserId;
+		FString Name;
 		FString Path;
 		
 		TOptional<float> FrameRate;
@@ -29,7 +29,7 @@ struct DATAINGESTCORE_API FIngestCaptureData
 	/** Audio info. */
 	struct FAudio
 	{
-		FString UserId;
+		FString Name;
 		FString Path;
 
 		TOptional<FString> TimecodeStart;
@@ -39,7 +39,7 @@ struct DATAINGESTCORE_API FIngestCaptureData
 	/** Calibration info. */
 	struct FCalibration
 	{
-		FString UserId;
+		FString Name;
 		FString Path;
 	};
 
diff --git a/Engine/Plugins/VirtualProduction/CaptureManager/CaptureManagerDevices/Source/MonoVideoLiveLinkDevice/Private/MonoVideoLiveLinkDevice.cpp b/Engine/Plugins/VirtualProduction/CaptureManager/CaptureManagerDevices/Source/MonoVideoLiveLinkDevice/Private/MonoVideoLiveLinkDevice.cpp
--- a/Engine/Plugins/VirtualProduction/CaptureManager/CaptureManagerDevices/Source/MonoVideoLiveLinkDevice/Private/MonoVideoLiveLinkDevice.cpp
+++ b/Engine/Plugins/VirtualProduction/CaptureManager/CaptureManagerDevices/Source/MonoVideoLiveLinkDevice/Private/MonoVideoLiveLinkDevice.cpp
@@ -251,7 +251,7 @@ TOptional<FTakeMetadata> UMonoVideoLiveLinkDevice::ReadTake(const FString& InCur
 		UE_LOG(LogMonoVideoLiveLinkDevice, Warning, TEXT("Failed to determine the frame rate for the video file %s"), *InCurrentTakeFile);
 	}
 
-	Video.UserId = MoveTemp(UserId);
+	Video.Name = MoveTemp(UserId);
 	Video.Format = FPaths::GetExtension(InCurrentTakeFile);
 	Video.Path = InCurrentTakeFile;
 	Video.PathType = FTakeMetadata::FVideo::EPathType::File;
@@ -265,7 +265,7 @@ TOptional<FTakeMetadata> UMonoVideoLiveLinkDevice::ReadTake(const FString& InCur
 	if (AudioReaderResult.HasValue())
 	{
 		FTakeMetadata::FAudio Audio;
-		Audio.UserId = TEXT("audio");
+		Audio.Name = TEXT("audio");
 		Audio.Path = InCurrentTakeFile;
 		Audio.Duration = static_cast<float>(AudioReaderResult.StealValue()->GetDuration().GetTotalSeconds());
 
diff --git a/Engine/Plugins/VirtualProduction/CaptureManager/CaptureManagerEditor/Source/CaptureManagerEditorSettings/Private/Settings/CaptureManagerEditorTemplateTokens.cpp b/Engine/Plugins/VirtualProduction/CaptureManager/CaptureManagerEditor/Source/CaptureManagerEditorSettings/Private/Settings/CaptureManagerEditorTemplateTokens.cpp
--- a/Engine/Plugins/VirtualProduction/CaptureManager/CaptureManagerEditor/Source/CaptureManagerEditorSettings/Private/Settings/CaptureManagerEditorTemplateTokens.cpp
+++ b/Engine/Plugins/VirtualProduction/CaptureManager/CaptureManagerEditor/Source/CaptureManagerEditorSettings/Private/Settings/CaptureManagerEditorTemplateTokens.cpp
@@ -55,7 +55,7 @@ UCaptureManagerVideoNamingTokens::UCaptureManagerVideoNamingTokens()
 {
 	Namespace = TEXT("cmv");
 
-	IngestVideoTokens.Emplace(FString(UE::CaptureManager::VideoTokens::NameKey), { FString(UE::CaptureManager::VideoTokens::NameKey), LOCTEXT("VideoName", "Name") });
+	IngestVideoTokens.Emplace(FString(UE::CaptureManager::VideoTokens::NameKey), { FString(UE::CaptureManager::VideoTokens::NameKey), LOCTEXT("VideoName", "Name (used to identify a track of recorded data)") });
 	IngestVideoTokens.Emplace(FString(UE::CaptureManager::VideoTokens::FrameRateKey), { FString(UE::CaptureManager::VideoTokens::FrameRateKey), LOCTEXT("VideoFrameRate", "Frame Rate") });
 }
 
@@ -99,7 +99,7 @@ UCaptureManagerAudioNamingTokens::UCaptureManagerAudioNamingTokens()
 {
 	Namespace = TEXT("cma");
 
-	IngestAudioTokens.Emplace(FString(UE::CaptureManager::AudioTokens::NameKey), { FString(UE::CaptureManager::AudioTokens::NameKey), LOCTEXT("AudioName", "Name") });
+	IngestAudioTokens.Emplace(FString(UE::CaptureManager::AudioTokens::NameKey), { FString(UE::CaptureManager::AudioTokens::NameKey), LOCTEXT("AudioName", "Name (used to identify a track of recorded data)") });
 }
 
 UE::CaptureManager::FIngestToken UCaptureManagerAudioNamingTokens::GetToken(const FString& InKey) const
@@ -142,7 +142,7 @@ UCaptureManagerCalibrationNamingTokens::UCaptureManagerCalibrationNamingTokens()
 {
 	Namespace = TEXT("cmc");
 
-	IngestCalibTokens.Emplace(FString(UE::CaptureManager::CalibTokens::NameKey), { FString(UE::CaptureManager::CalibTokens::NameKey), LOCTEXT("CalibrationName", "Name") });
+	IngestCalibTokens.Emplace(FString(UE::CaptureManager::CalibTokens::NameKey), { FString(UE::CaptureManager::CalibTokens::NameKey), LOCTEXT("CalibrationName", "Name (used to identify a track of recorded data)") });
 }
 
 UE::CaptureManager::FIngestToken UCaptureManagerCalibrationNamingTokens::GetToken(const FString& InKey) const
diff --git a/Engine/Plugins/VirtualProduction/CaptureManager/CaptureManagerEditor/Source/LiveLinkHubWorkerManager/Private/IngestProcess/IngestCaptureDataProcess.cpp b/Engine/Plugins/VirtualProduction/CaptureManager/CaptureManagerEditor/Source/LiveLinkHubWorkerManager/Private/IngestProcess/IngestCaptureDataProcess.cpp
--- a/Engine/Plugins/VirtualProduction/CaptureManager/CaptureManagerEditor/Source/LiveLinkHubWorkerManager/Private/IngestProcess/IngestCaptureDataProcess.cpp
+++ b/Engine/Plugins/VirtualProduction/CaptureManager/CaptureManagerEditor/Source/LiveLinkHubWorkerManager/Private/IngestProcess/IngestCaptureDataProcess.cpp
@@ -181,7 +181,7 @@ UE::CaptureManager::FCreateAssetsData FIngestCaptureDataProcess::PrepareAssetsDa
 				FStringFormatNamedArguments VideoNamedArgs;
 				{
 					// Video tokens
-					VideoNamedArgs.Add(Settings->GetVideoNamingTokens()->GetToken(FString(VideoTokens::NameKey)).Name, Video.UserId);
+					VideoNamedArgs.Add(Settings->GetVideoNamingTokens()->GetToken(FString(VideoTokens::NameKey)).Name, Video.Name);
 					VideoNamedArgs.Add(Settings->GetVideoNamingTokens()->GetToken(FString(VideoTokens::FrameRateKey)).Name, FString::Printf(TEXT("%.2f"), FrameRate.AsDecimal()));
 				}
 				// Evaluate Video Settings
@@ -193,7 +193,7 @@ UE::CaptureManager::FCreateAssetsData FIngestCaptureDataProcess::PrepareAssetsDa
 				ImageSequenceData.AssetName = ProcessedImageSequenceAssetName;
 
 				ImageSequenceData.FrameRate = FrameRate;
-				ImageSequenceData.Name = Video.UserId;
+				ImageSequenceData.Name = Video.Name;
 				ImageSequenceData.SequenceDirectory = Video.Path;
 				ImageSequenceData.bTimecodePresent = Video.TimecodeStart.IsSet();
 				ImageSequenceData.Timecode = Video.TimecodeStart.IsSet() ? ParseTimecode(Video.TimecodeStart.GetValue()) : FTimecode();
@@ -213,7 +213,7 @@ UE::CaptureManager::FCreateAssetsData FIngestCaptureDataProcess::PrepareAssetsDa
 				FStringFormatNamedArguments DepthNamedArgs;
 				{
 					// Depth tokens
-					DepthNamedArgs.Add(Settings->GetVideoNamingTokens()->GetToken(FString(VideoTokens::NameKey)).Name, Depth.UserId);
+					DepthNamedArgs.Add(Settings->GetVideoNamingTokens()->GetToken(FString(VideoTokens::NameKey)).Name, Depth.Name);
 					DepthNamedArgs.Add(Settings->GetVideoNamingTokens()->GetToken(FString(VideoTokens::FrameRateKey)).Name, FString::Printf(TEXT("%.2f"), FrameRate.AsDecimal()));
 				}
 				// Evaluate Depth Settings
@@ -226,7 +226,7 @@ UE::CaptureManager::FCreateAssetsData FIngestCaptureDataProcess::PrepareAssetsDa
 				DepthSequenceData.AssetName = ProcessedDepthSequenceAssetName;
 
 				DepthSequenceData.FrameRate = FrameRate;
-				DepthSequenceData.Name = Depth.UserId;
+				DepthSequenceData.Name = Depth.Name;
 				DepthSequenceData.SequenceDirectory = Depth.Path;
 				DepthSequenceData.bTimecodePresent = Depth.TimecodeStart.IsSet();
 				DepthSequenceData.Timecode = Depth.TimecodeStart.IsSet() ? ParseTimecode(Depth.TimecodeStart.GetValue()) : FTimecode();
@@ -249,7 +249,7 @@ UE::CaptureManager::FCreateAssetsData FIngestCaptureDataProcess::PrepareAssetsDa
 				FStringFormatNamedArguments AudioNamedArgs;
 				{
 					// Audio tokens
-					AudioNamedArgs.Add(Settings->GetAudioNamingTokens()->GetToken(FString(AudioTokens::NameKey)).Name, Audio.UserId);
+					AudioNamedArgs.Add(Settings->GetAudioNamingTokens()->GetToken(FString(AudioTokens::NameKey)).Name, Audio.Name);
 				}
 				// Evaluate Audio Settings
 				FString AudioAssetName = FString::Format(*Settings->SoundwaveAssetName, AudioNamedArgs);
@@ -260,7 +260,7 @@ UE::CaptureManager::FCreateAssetsData FIngestCaptureDataProcess::PrepareAssetsDa
 				SanitizeAssetName(ProcessedAudioAssetName);
 				AudioData.AssetName = ProcessedAudioAssetName;
 
-				AudioData.Name = Audio.UserId;
+				AudioData.Name = Audio.Name;
 				AudioData.WAVFile = Audio.Path;
 				AudioData.bTimecodePresent = Audio.TimecodeStart.IsSet();
 				AudioData.Timecode = Audio.TimecodeStart.IsSet() ? ParseTimecode(Audio.TimecodeStart.GetValue()) : FTimecode();
