[CaptureManager] Added operator overload for take metadata schema version

diff --git a/Engine/Plugins/VirtualProduction/CaptureManager/CaptureManagerCore/Source/CaptureManagerTakeMetadata/Private/CaptureManagerTakeMetadata.cpp b/Engine/Plugins/VirtualProduction/CaptureManager/CaptureManagerCore/Source/CaptureManagerTakeMetadata/Private/CaptureManagerTakeMetadata.cpp
--- a/Engine/Plugins/VirtualProduction/CaptureManager/CaptureManagerCore/Source/CaptureManagerTakeMetadata/Private/CaptureManagerTakeMetadata.cpp
+++ b/Engine/Plugins/VirtualProduction/CaptureManager/CaptureManagerCore/Source/CaptureManagerTakeMetadata/Private/CaptureManagerTakeMetadata.cpp
@@ -80,13 +80,13 @@ static FTakeMetadata::FDevice ParseDevice(const rapidjson::Value& InDeviceObject
 	FTakeMetadata::FDevice TakeDevice;
 	TakeDevice.Type = InDeviceObject["Type"].GetString();
 	TakeDevice.Model = InDeviceObject["Model"].GetString();
-	if (InVersion.Major >= 4 && InVersion.Minor >= 2)
+	if (InVersion < FTakeMetadata::FSchemaVersion{ 4, 2 })
 	{
-		TakeDevice.Name = InDeviceObject["Name"].GetString();
+		TakeDevice.Name = InDeviceObject["UserId"].GetString();
 	}
 	else
 	{
-		TakeDevice.Name = InDeviceObject["UserId"].GetString();
+		TakeDevice.Name = InDeviceObject["Name"].GetString();
 	}
 
 	if (InDeviceObject.HasMember("Platform"))
@@ -134,13 +134,13 @@ static FTakeMetadata::FVideo ParseVideo(const rapidjson::Value& InVideoObject, c
 {
 	FTakeMetadata::FVideo Video;
 
-	if (InVersion.Major >= 4 && InVersion.Minor >= 2)
+	if (InVersion < FTakeMetadata::FSchemaVersion{ 4, 2 })
 	{
-		Video.Name = InVideoObject["Name"].GetString();
+		Video.Name = InVideoObject["UserId"].GetString();
 	}
 	else
 	{
-		Video.Name = InVideoObject["UserId"].GetString();
+		Video.Name = InVideoObject["Name"].GetString();
 	}
 	Video.Path = InVideoObject["Path"].GetString();
 
@@ -240,13 +240,13 @@ static TArray<FTakeMetadata::FCalibration> ParseCalibration(const rapidjson::Val
 		for (const rapidjson::Value& CalibrationObject : InCalibrationArray.GetArray())
 		{
 			FTakeMetadata::FCalibration Calibration;
-			if (InVersion.Major >= 4 && InVersion.Minor >= 2)
+			if (InVersion < FTakeMetadata::FSchemaVersion{ 4, 2 })
 			{
-				Calibration.Name = CalibrationObject["Name"].GetString();
+				Calibration.Name = CalibrationObject["UserId"].GetString();
 			}
 			else
 			{
-				Calibration.Name = CalibrationObject["UserId"].GetString();
+				Calibration.Name = CalibrationObject["Name"].GetString();
 			}
 			Calibration.Format = CalibrationObject["Format"].GetString();
 			Calibration.Path = CalibrationObject["Path"].GetString();
@@ -263,13 +263,13 @@ static TArray<FTakeMetadata::FAudio> ParseAudio(const rapidjson::Value& InAudioA
 	for (const rapidjson::Value& AudioObject : InAudioArray.GetArray())
 	{
 		FTakeMetadata::FAudio Audio;
-		if (InVersion.Major >= 4 && InVersion.Minor >= 2)
+		if (InVersion < FTakeMetadata::FSchemaVersion{ 4, 2 })
 		{
-			Audio.Name = AudioObject["Name"].GetString();
+			Audio.Name = AudioObject["UserId"].GetString();
 		}
 		else
 		{
-			Audio.Name = AudioObject["UserId"].GetString();
+			Audio.Name = AudioObject["Name"].GetString();
 		}
 
 		Audio.Path = AudioObject["Path"].GetString();
diff --git a/Engine/Plugins/VirtualProduction/CaptureManager/CaptureManagerCore/Source/CaptureManagerTakeMetadata/Public/CaptureManagerTakeMetadata.h b/Engine/Plugins/VirtualProduction/CaptureManager/CaptureManagerCore/Source/CaptureManagerTakeMetadata/Public/CaptureManagerTakeMetadata.h
--- a/Engine/Plugins/VirtualProduction/CaptureManager/CaptureManagerCore/Source/CaptureManagerTakeMetadata/Public/CaptureManagerTakeMetadata.h
+++ b/Engine/Plugins/VirtualProduction/CaptureManager/CaptureManagerCore/Source/CaptureManagerTakeMetadata/Public/CaptureManagerTakeMetadata.h
@@ -48,6 +48,11 @@ public:
 
 		bool operator==(const FSchemaVersion& rhs) const = default;
 
+		bool operator<(const FSchemaVersion& rhs) const
+		{
+			return Major < rhs.Major || (Major == rhs.Major && Minor < rhs.Minor);
+		}
+
 		friend uint32 GetTypeHash(const FSchemaVersion& InVersion)
 		{
 			return HashCombine(GetTypeHash(InVersion.Major), GetTypeHash(InVersion.Minor));
