Sending an update progress when it is updated by 3

diff --git a/Engine/Plugins/VirtualProduction/CaptureManager/CaptureManagerCore/Source/CaptureUtils/Private/Async/TaskProgress.cpp b/Engine/Plugins/VirtualProduction/CaptureManager/CaptureManagerCore/Source/CaptureUtils/Private/Async/TaskProgress.cpp
--- a/Engine/Plugins/VirtualProduction/CaptureManager/CaptureManagerCore/Source/CaptureUtils/Private/Async/TaskProgress.cpp
+++ b/Engine/Plugins/VirtualProduction/CaptureManager/CaptureManagerCore/Source/CaptureUtils/Private/Async/TaskProgress.cpp
@@ -74,6 +74,11 @@ FTaskProgress::FTaskProgress(uint32 InAmountOfWork, FProgressReporter InReporter
 	CurrentProgressValues.Lock()->AddZeroed(InAmountOfWork);
 }
 
+void FTaskProgress::SetReportThreshold(double InReportThresholdPercent)
+{
+	ReportThreshold = InReportThresholdPercent / 100.0;
+}
+
 FTaskProgress::FTask FTaskProgress::StartTask()
 {
 	int32 NewTask = ++CurrentTask;
@@ -107,10 +112,33 @@ void FTaskProgress::Update(int32 InTaskToUpdate, double InProgress)
 	(*ProgressValues)[InTaskToUpdate] = InProgress;
 }
 
-void FTaskProgress::Report() const
+void FTaskProgress::Report()
 {
 	float TotalProgress = GetTotalProgress();
-	Reporter.ExecuteIfBound(TotalProgress);
+	bool bShouldReport = false;
+
+	if (!FMath::IsNearlyZero(ReportThreshold))
+	{	
+		if (FMath::IsNearlyEqual(TotalProgress, 1.0))
+		{
+			bShouldReport = true;
+		} 
+		else if (TotalProgress > LastReportedValue && 
+			(TotalProgress - LastReportedValue) > ReportThreshold)
+		{
+			bShouldReport = true;
+		}
+	}
+	else
+	{
+		bShouldReport = true;
+	}
+
+	if (bShouldReport)
+	{
+		Reporter.ExecuteIfBound(TotalProgress);
+		LastReportedValue = TotalProgress;
+	}
 }
 
 }
\ No newline at end of file
diff --git a/Engine/Plugins/VirtualProduction/CaptureManager/CaptureManagerCore/Source/CaptureUtils/Public/Async/TaskProgress.h b/Engine/Plugins/VirtualProduction/CaptureManager/CaptureManagerCore/Source/CaptureUtils/Public/Async/TaskProgress.h
--- a/Engine/Plugins/VirtualProduction/CaptureManager/CaptureManagerCore/Source/CaptureUtils/Public/Async/TaskProgress.h
+++ b/Engine/Plugins/VirtualProduction/CaptureManager/CaptureManagerCore/Source/CaptureUtils/Public/Async/TaskProgress.h
@@ -42,6 +42,7 @@ public:
 	};
 
 	FTaskProgress(uint32 InAmountOfWork, FProgressReporter InReport);
+	void SetReportThreshold(double InReportThresholdPercent);
 
 	FTask StartTask();
 
@@ -50,11 +51,14 @@ public:
 private:
 
 	void Update(int32 InTaskToUpdate, double InProgress);
-	void Report() const;
+	void Report();
 
 	FProgressReporter Reporter;
 
 	std::atomic<int32> CurrentTask;
+	double ReportThreshold = 0.0;
+
+	std::atomic<double> LastReportedValue = 0.0;
 
 	mutable FMonitor<TArray<double>> CurrentProgressValues;
 };
diff --git a/Engine/Plugins/VirtualProduction/CaptureManager/CaptureManagerEditor/Source/LiveLinkHubWorkerManager/Private/LiveLinkHubWorker.cpp b/Engine/Plugins/VirtualProduction/CaptureManager/CaptureManagerEditor/Source/LiveLinkHubWorkerManager/Private/LiveLinkHubWorker.cpp
--- a/Engine/Plugins/VirtualProduction/CaptureManager/CaptureManagerEditor/Source/LiveLinkHubWorkerManager/Private/LiveLinkHubWorker.cpp
+++ b/Engine/Plugins/VirtualProduction/CaptureManager/CaptureManagerEditor/Source/LiveLinkHubWorkerManager/Private/LiveLinkHubWorker.cpp
@@ -167,6 +167,7 @@ bool FLiveLinkHubImportWorker::HandleTakeDownload(FUploadDataHeader InHeader, TS
 			{
 				PublishProgressEvent(SharedMessenger, CaptureSourceId, TakeUploadId, InTotalProgress);
 			}));
+	TaskProgress->SetReportThreshold(3.0);
 
 	FTaskProgress::FTask Task = TaskProgress->StartTask();
 
