Updated headers using UnrealCodeFixup to make sure dllstorage is on methods/staticvars instead of types. Part 4/n

diff --git a/Engine/Plugins/VirtualProduction/CameraCalibrationCore/Source/CameraCalibrationCoreEditor/Public/LensDistortionStateDetailCustomization.h b/Engine/Plugins/VirtualProduction/CameraCalibrationCore/Source/CameraCalibrationCoreEditor/Public/LensDistortionStateDetailCustomization.h
--- a/Engine/Plugins/VirtualProduction/CameraCalibrationCore/Source/CameraCalibrationCoreEditor/Public/LensDistortionStateDetailCustomization.h
+++ b/Engine/Plugins/VirtualProduction/CameraCalibrationCore/Source/CameraCalibrationCoreEditor/Public/LensDistortionStateDetailCustomization.h
@@ -7,26 +7,30 @@
 #include "Models/LensModel.h"
 #include "Templates/SubclassOf.h"
 
-class CAMERACALIBRATIONCOREEDITOR_API FLensDistortionStateDetailCustomization : public IPropertyTypeCustomization
+#define UE_API CAMERACALIBRATIONCOREEDITOR_API
+
+class FLensDistortionStateDetailCustomization : public IPropertyTypeCustomization
 {
 public:
-	FLensDistortionStateDetailCustomization(TSubclassOf<ULensModel> InLensModel);
+	UE_API FLensDistortionStateDetailCustomization(TSubclassOf<ULensModel> InLensModel);
 
 	// Begin IDetailCustomization interface
-	virtual void CustomizeHeader(TSharedRef<IPropertyHandle> PropertyHandle, FDetailWidgetRow& HeaderRow, IPropertyTypeCustomizationUtils& CustomizationUtils) override;
-	virtual void CustomizeChildren(TSharedRef<IPropertyHandle> PropertyHandle, IDetailChildrenBuilder& ChildBuilder, IPropertyTypeCustomizationUtils& CustomizationUtils) override;
+	UE_API virtual void CustomizeHeader(TSharedRef<IPropertyHandle> PropertyHandle, FDetailWidgetRow& HeaderRow, IPropertyTypeCustomizationUtils& CustomizationUtils) override;
+	UE_API virtual void CustomizeChildren(TSharedRef<IPropertyHandle> PropertyHandle, IDetailChildrenBuilder& ChildBuilder, IPropertyTypeCustomizationUtils& CustomizationUtils) override;
 	//~ End IDetailCustomization interface
 
 private:
 	/** Customize the FDistortionInfo struct property */
-	void CustomizeDistortionInfo(TSharedRef<IPropertyHandle> PropertyHandle, IDetailChildrenBuilder& ChildBuilder, IPropertyTypeCustomizationUtils& CustomizationUtils);
+	UE_API void CustomizeDistortionInfo(TSharedRef<IPropertyHandle> PropertyHandle, IDetailChildrenBuilder& ChildBuilder, IPropertyTypeCustomizationUtils& CustomizationUtils);
 
 	/** Customize the FFocalLengthInfo struct property */
-	void CustomizeFocalLengthInfo(TSharedRef<IPropertyHandle> PropertyHandle, IDetailChildrenBuilder& ChildBuilder, IPropertyTypeCustomizationUtils& CustomizationUtils);
+	UE_API void CustomizeFocalLengthInfo(TSharedRef<IPropertyHandle> PropertyHandle, IDetailChildrenBuilder& ChildBuilder, IPropertyTypeCustomizationUtils& CustomizationUtils);
 
 	/** Customize the FImageCenterInfo struct property */
-	void CustomizeImageCenterInfo(TSharedRef<IPropertyHandle> PropertyHandle, IDetailChildrenBuilder& ChildBuilder, IPropertyTypeCustomizationUtils& CustomizationUtils);
+	UE_API void CustomizeImageCenterInfo(TSharedRef<IPropertyHandle> PropertyHandle, IDetailChildrenBuilder& ChildBuilder, IPropertyTypeCustomizationUtils& CustomizationUtils);
 
 	/** The LensModel used to customize the names of the properties in the distortion parameter array */
 	TSubclassOf<ULensModel> LensModel;
 };
+
+#undef UE_API
diff --git a/Engine/Plugins/VirtualProduction/CaptureData/Source/CaptureDataCore/Public/CameraCalibration.h b/Engine/Plugins/VirtualProduction/CaptureData/Source/CaptureDataCore/Public/CameraCalibration.h
--- a/Engine/Plugins/VirtualProduction/CaptureData/Source/CaptureDataCore/Public/CameraCalibration.h
+++ b/Engine/Plugins/VirtualProduction/CaptureData/Source/CaptureDataCore/Public/CameraCalibration.h
@@ -9,6 +9,8 @@
 
 #include "CameraCalibration.generated.h"
 
+#define UE_API CAPTUREDATACORE_API
+
 
 USTRUCT(BlueprintType)
 struct FExtendedLensFile
@@ -41,7 +43,7 @@ public:
 	uint32 CameraIndex2 {};
 };
 
-struct CAPTUREDATACORE_API FCameraCalibration
+struct FCameraCalibration
 {
 public:
 
@@ -80,16 +82,16 @@ public:
 *   used in footage for MetaHuman Identity and Performance assets.
 * 
 **/
-UCLASS(BlueprintType)
-class CAPTUREDATACORE_API UCameraCalibration : public UObject
+UCLASS(MinimalAPI, BlueprintType)
+class UCameraCalibration : public UObject
 {
 	GENERATED_BODY()
 
 public:
 	//~Begin UObject interface
-	virtual void PostInitProperties() override;
-	virtual void PostLoad() override;
-	virtual void GetAssetRegistryTags(TArray<FAssetRegistryTag>& OutTags) const override;
+	UE_API virtual void PostInitProperties() override;
+	UE_API virtual void PostLoad() override;
+	UE_API virtual void GetAssetRegistryTags(TArray<FAssetRegistryTag>& OutTags) const override;
 	//~End UObject interface
 
 	UPROPERTY(EditAnywhere, Category = Calibration)
@@ -104,18 +106,20 @@ public:
 	TObjectPtr<UAssetImportData> AssetImportData;
 #endif // WITH_EDITORONLY_DATA
 
-	bool ConvertToTrackerNodeCameraModels(TArray<FCameraCalibration>& OutCalibrations,
+	UE_API bool ConvertToTrackerNodeCameraModels(TArray<FCameraCalibration>& OutCalibrations,
 										  TArray<TPair<FString, FString>>& OutStereoReconstructionPairs) const;
 
-	bool ConvertFromTrackerNodeCameraModels(const TArray<FCameraCalibration>& InCalibrations, bool bInUsingUnrealCoordinateSystem = false);
-	bool ConvertFromTrackerNodeCameraModels(const TArray<FCameraCalibration>& InCalibrations, const TMap<FString, FString>& InLensAssetNamesMap, bool bInUsingUnrealCoordinateSystem = false);
+	UE_API bool ConvertFromTrackerNodeCameraModels(const TArray<FCameraCalibration>& InCalibrations, bool bInUsingUnrealCoordinateSystem = false);
+	UE_API bool ConvertFromTrackerNodeCameraModels(const TArray<FCameraCalibration>& InCalibrations, const TMap<FString, FString>& InLensAssetNamesMap, bool bInUsingUnrealCoordinateSystem = false);
 
-	int32 GetCalibrationIndexByName(const FString& InName) const;
+	UE_API int32 GetCalibrationIndexByName(const FString& InName) const;
 
 private:
 
-	void CreateLensFileForCalibration(const FCameraCalibration& InCalibration, const FString& InAssetName, bool bInUsingUnrealCoordinateSystem);
-	void AddStereoPair();
+	UE_API void CreateLensFileForCalibration(const FCameraCalibration& InCalibration, const FString& InAssetName, bool bInUsingUnrealCoordinateSystem);
+	UE_API void AddStereoPair();
 	
 	EMediaOrientation CameraOrientation;
-};
\ No newline at end of file
+};
+
+#undef UE_API
diff --git a/Engine/Plugins/VirtualProduction/CaptureData/Source/CaptureDataCore/Public/CaptureData.h b/Engine/Plugins/VirtualProduction/CaptureData/Source/CaptureDataCore/Public/CaptureData.h
--- a/Engine/Plugins/VirtualProduction/CaptureData/Source/CaptureDataCore/Public/CaptureData.h
+++ b/Engine/Plugins/VirtualProduction/CaptureData/Source/CaptureDataCore/Public/CaptureData.h
@@ -13,6 +13,8 @@
 
 #include "CaptureData.generated.h"
 
+#define UE_API CAPTUREDATACORE_API
+
 /////////////////////////////////////////////////////
 // UCaptureData
 
@@ -69,18 +71,18 @@ private:
 *   of that person in MetaHuman Performance asset
 *
 **/
-UCLASS()
-class CAPTUREDATACORE_API UMeshCaptureData : public UCaptureData
+UCLASS(MinimalAPI)
+class UMeshCaptureData : public UCaptureData
 {
 	GENERATED_BODY()
 
 public:
 	// UCaptureData interface
-	virtual bool IsInitialized(EInitializedCheck InInitializedCheck) const override;
+	UE_API virtual bool IsInitialized(EInitializedCheck InInitializedCheck) const override;
 
 public:
 	/** Gets the data in the format expected by the face fitting API */
-	void GetDataForConforming(const FTransform& InTransform, TArray<float>& OutVertices, TArray<int32>& OutTriangles) const;
+	UE_API void GetDataForConforming(const FTransform& InTransform, TArray<float>& OutVertices, TArray<int32>& OutTriangles) const;
 
 	// The target mesh for conforming. This can be either a Static or SkeletalMesh
 	UPROPERTY(EditAnywhere, Category = "Capture", Meta = (AllowedClasses = "/Script/Engine.StaticMesh, /Script/Engine.SkeletalMesh"))
@@ -103,7 +105,7 @@ enum class EFootageDeviceClass : uint8
 };
 
 USTRUCT(BlueprintType)
-struct CAPTUREDATACORE_API FFootageCaptureMetadata
+struct FFootageCaptureMetadata
 {
 	GENERATED_BODY()
 
@@ -126,7 +128,7 @@ struct CAPTUREDATACORE_API FFootageCaptureMetadata
 	FString DeviceModelName;
 
 	/** Sets the DeviceClass property parsing the model string accordingly */
-	void SetDeviceClass(const FString& InDeviceModel);
+	UE_API void SetDeviceClass(const FString& InDeviceModel);
 
 private:
 	struct FIosDeviceVersion
@@ -138,8 +140,8 @@ private:
 			: Major(Major), Minor(Minor)
 		{}
 	};
-	TOptional<FIosDeviceVersion> ParseIosDeviceVersion(const FString& Prefix, const FString& ModelName);
-	EFootageDeviceClass IPhoneDeviceClass(TOptional<FIosDeviceVersion> IosDeviceVersion);
+	UE_API TOptional<FIosDeviceVersion> ParseIosDeviceVersion(const FString& Prefix, const FString& ModelName);
+	UE_API EFootageDeviceClass IPhoneDeviceClass(TOptional<FIosDeviceVersion> IosDeviceVersion);
 };
 
 USTRUCT(BlueprintType)
@@ -193,8 +195,8 @@ enum class ETimecodeAlignment : int32
  *  of that person in MetaHuman Performance asset
  *
 **/
-UCLASS()
-class CAPTUREDATACORE_API UFootageCaptureData : public UCaptureData
+UCLASS(MinimalAPI)
+class UFootageCaptureData : public UCaptureData
 {
 	GENERATED_BODY()
 
@@ -203,11 +205,11 @@ public:
 	using FVerifyResult = TValueOrError<void, FString>;
 
 	//~Begin UCaptureData interface
-	virtual bool IsInitialized(EInitializedCheck InInitializedCheck) const override;
+	UE_API virtual bool IsInitialized(EInitializedCheck InInitializedCheck) const override;
 	//~End UCaptureData interface
 
 	//~Begin UObject interface
-	virtual void PostLoad() override;
+	UE_API virtual void PostLoad() override;
 	//~End UObject interface
 
 public:
@@ -254,25 +256,25 @@ public:
 public:
 
 	/** Gets the resolution of the color channel */
-	FIntPoint GetFootageColorResolution() const;
+	UE_API FIntPoint GetFootageColorResolution() const;
 
-	void GetFrameRanges(const FFrameRate& InTargetRate, ETimecodeAlignment InTimecodeAlignment, bool bInIncludeAudio, TMap<TWeakObjectPtr<UObject>, TRange<FFrameNumber>>& OutMediaFrameRanges, TRange<FFrameNumber>& OutProcessingFrameRange, TRange<FFrameNumber>& OutMaximumFrameRange) const;
+	UE_API void GetFrameRanges(const FFrameRate& InTargetRate, ETimecodeAlignment InTimecodeAlignment, bool bInIncludeAudio, TMap<TWeakObjectPtr<UObject>, TRange<FFrameNumber>>& OutMediaFrameRanges, TRange<FFrameNumber>& OutProcessingFrameRange, TRange<FFrameNumber>& OutMaximumFrameRange) const;
 
-	static TRange<FFrameNumber> GetAudioFrameRange(const FFrameRate& InTargetRate, ETimecodeAlignment InTimecodeAlignment, USoundWave* InMedia, const FTimecode& InMediaTimecode, const FFrameRate& InMediaTimecodeRate);
+	static UE_API TRange<FFrameNumber> GetAudioFrameRange(const FFrameRate& InTargetRate, ETimecodeAlignment InTimecodeAlignment, USoundWave* InMedia, const FTimecode& InMediaTimecode, const FFrameRate& InMediaTimecodeRate);
 
 	/** Functions for getting the effective timecode info on a track - note some tracks may not have a timecode set
 	*   in which case their effective timecode is defined by others tracks that do have a timecode set.
 	*/
-	FTimecode GetEffectiveImageTimecode(int32 InView) const;
-	FFrameRate GetEffectiveImageTimecodeRate(int32 InView) const;
+	UE_API FTimecode GetEffectiveImageTimecode(int32 InView) const;
+	UE_API FFrameRate GetEffectiveImageTimecodeRate(int32 InView) const;
 
-	FTimecode GetEffectiveDepthTimecode(int32 InView) const;
-	FFrameRate GetEffectiveDepthTimecodeRate(int32 InView) const;
+	UE_API FTimecode GetEffectiveDepthTimecode(int32 InView) const;
+	UE_API FFrameRate GetEffectiveDepthTimecodeRate(int32 InView) const;
 
-	FTimecode GetEffectiveAudioTimecode() const;
-	FFrameRate GetEffectiveAudioTimecodeRate() const;
+	UE_API FTimecode GetEffectiveAudioTimecode() const;
+	UE_API FFrameRate GetEffectiveAudioTimecodeRate() const;
 
-	FVerifyResult VerifyData(EInitializedCheck InInitializedCheck) const;
+	UE_API FVerifyResult VerifyData(EInitializedCheck InInitializedCheck) const;
 
 	struct FPathAssociation
 	{
@@ -283,18 +285,20 @@ public:
 	};
 
 	/** Check the existence of the parent directory for all image sequences and return a list of those which fail the test. */
-	[[nodiscard]] TArray<FPathAssociation> CheckImageSequencePaths() const;
+	[[nodiscard]] UE_API TArray<FPathAssociation> CheckImageSequencePaths() const;
 
 	/** List of all RGB cameras (views) in the footage capture data. Ensures InOutCamera is a valid camera name. */
-	static void PopulateCameraNames(UFootageCaptureData* InFootageCaptureData, FString& InOutCamera, TArray<TSharedPtr<FString>>& OutCameraNames);
+	static UE_API void PopulateCameraNames(UFootageCaptureData* InFootageCaptureData, FString& InOutCamera, TArray<TSharedPtr<FString>>& OutCameraNames);
 
-	int32 GetViewIndexByCameraName(const FString& InName) const;
+	UE_API int32 GetViewIndexByCameraName(const FString& InName) const;
 
 private:
 	/** Checks if camera calibration and views are valid. */
-	FVerifyResult ViewsContainsValidData(EInitializedCheck InInitializedCheck) const;
-	FVerifyResult MetadataContainsValidData() const;
-	FVerifyResult CalibrationContainsValidData() const;
+	UE_API FVerifyResult ViewsContainsValidData(EInitializedCheck InInitializedCheck) const;
+	UE_API FVerifyResult MetadataContainsValidData() const;
+	UE_API FVerifyResult CalibrationContainsValidData() const;
+
+	UE_API void GetDefaultTimecodeInfo(FTimecode& OutTimecode, FFrameRate& OutFrameRate) const;
+};
 
-	void GetDefaultTimecodeInfo(FTimecode& OutTimecode, FFrameRate& OutFrameRate) const;
-};
\ No newline at end of file
+#undef UE_API
diff --git a/Engine/Plugins/VirtualProduction/CaptureData/Source/CaptureDataCore/Public/FrameRange.h b/Engine/Plugins/VirtualProduction/CaptureData/Source/CaptureDataCore/Public/FrameRange.h
--- a/Engine/Plugins/VirtualProduction/CaptureData/Source/CaptureDataCore/Public/FrameRange.h
+++ b/Engine/Plugins/VirtualProduction/CaptureData/Source/CaptureDataCore/Public/FrameRange.h
@@ -6,15 +6,17 @@
 
 #include "FrameRange.generated.h"
 
+#define UE_API CAPTUREDATACORE_API
+
 USTRUCT(BlueprintType)
-struct CAPTUREDATACORE_API FFrameRange
+struct FFrameRange
 {
 
 public:
 
 	GENERATED_BODY()
 
-	bool operator==(const FFrameRange& InOther) const;
+	UE_API bool operator==(const FFrameRange& InOther) const;
 
 	UPROPERTY(EditAnywhere, Category = "Frame Range")
 	FString Name = TEXT("Unnamed");
@@ -25,7 +27,7 @@ public:
 	UPROPERTY(EditAnywhere, Category = "Frame Range")
 	int32 EndFrame = -1;
 
-	static bool ContainsFrame(int32 InFrame, const TArray<FFrameRange>& InFrameRangeArray);
+	static UE_API bool ContainsFrame(int32 InFrame, const TArray<FFrameRange>& InFrameRangeArray);
 };
 
 UENUM()
@@ -43,3 +45,4 @@ typedef TMap<EFrameRangeType, TArray<FFrameRange>> FFrameRangeMap;
 /** Packs a list of frame numbers into a sorted array of frame ranges. Essentially collapsing consecutive frame numbers into blocks */
 CAPTUREDATACORE_API TArray<FFrameRange> PackIntoFrameRanges(TArray<FFrameNumber> InFrameNumbers);
 
+#undef UE_API
diff --git a/Engine/Plugins/VirtualProduction/CaptureData/Source/CaptureDataCore/Public/ImageSequencePathChecker.h b/Engine/Plugins/VirtualProduction/CaptureData/Source/CaptureDataCore/Public/ImageSequencePathChecker.h
--- a/Engine/Plugins/VirtualProduction/CaptureData/Source/CaptureDataCore/Public/ImageSequencePathChecker.h
+++ b/Engine/Plugins/VirtualProduction/CaptureData/Source/CaptureDataCore/Public/ImageSequencePathChecker.h
@@ -6,17 +6,19 @@
 
 #include "Internationalization/Internationalization.h"
 
+#define UE_API CAPTUREDATACORE_API
+
 namespace UE::CaptureData
 {
 
-class CAPTUREDATACORE_API FImageSequencePathChecker
+class FImageSequencePathChecker
 {
 public:
-	explicit FImageSequencePathChecker(FText InAssetDisplayName);
+	UE_API explicit FImageSequencePathChecker(FText InAssetDisplayName);
 
-	void Check(const UFootageCaptureData& InCaptureData);
-	void DisplayDialog() const;
-	bool HasError() const;
+	UE_API void Check(const UFootageCaptureData& InCaptureData);
+	UE_API void DisplayDialog() const;
+	UE_API bool HasError() const;
 
 private:
 	int32 NumCaptureDataFootageAssets;
@@ -25,3 +27,5 @@ private:
 };
 
 }
+
+#undef UE_API
diff --git a/Engine/Plugins/VirtualProduction/CaptureData/Source/CaptureDataCore/Public/ImageSequenceUtils.h b/Engine/Plugins/VirtualProduction/CaptureData/Source/CaptureDataCore/Public/ImageSequenceUtils.h
--- a/Engine/Plugins/VirtualProduction/CaptureData/Source/CaptureDataCore/Public/ImageSequenceUtils.h
+++ b/Engine/Plugins/VirtualProduction/CaptureData/Source/CaptureDataCore/Public/ImageSequenceUtils.h
@@ -4,19 +4,23 @@
 
 #include "CoreMinimal.h"
 
+#define UE_API CAPTUREDATACORE_API
+
 /** Utility functions for getting information about image sequences. */
-class CAPTUREDATACORE_API FImageSequenceUtils
+class FImageSequenceUtils
 {
 public:
 	/** Get image sequence path and list of files from image media source asset. */
-	static bool GetImageSequencePathAndFilesFromAsset(const class UImgMediaSource* InImgSequence, FString& OutFullSequencePath, TArray<FString>& OutImageFiles);
+	static UE_API bool GetImageSequencePathAndFilesFromAsset(const class UImgMediaSource* InImgSequence, FString& OutFullSequencePath, TArray<FString>& OutImageFiles);
 
 	/** Get list of image file paths from the image sequence path. */
-	static bool GetImageSequenceFilesFromPath(const FString& InFullSequencePath, TArray<FString>& OutImageFiles);
+	static UE_API bool GetImageSequenceFilesFromPath(const FString& InFullSequencePath, TArray<FString>& OutImageFiles);
 
 	/** Get image sequence info (e.g. dimensions, number of images) from image media source asset. */
-	static bool GetImageSequenceInfoFromAsset(const class UImgMediaSource* InImgSequence, FIntVector2& OutDimensions, int32& OutNumImages);
+	static UE_API bool GetImageSequenceInfoFromAsset(const class UImgMediaSource* InImgSequence, FIntVector2& OutDimensions, int32& OutNumImages);
 
 	/** Get image sequence info (e.g. dimensions, number of images) from image sequence path. */
-	static bool GetImageSequenceInfoFromPath(const FString& InFullSequencePath, FIntVector2& OutDimensions, int32& OutNumImages);
+	static UE_API bool GetImageSequenceInfoFromPath(const FString& InFullSequencePath, FIntVector2& OutDimensions, int32& OutNumImages);
 };
+
+#undef UE_API
diff --git a/Engine/Plugins/VirtualProduction/CaptureData/Source/CaptureDataCore/Public/OpenCVHelperLocal.h b/Engine/Plugins/VirtualProduction/CaptureData/Source/CaptureDataCore/Public/OpenCVHelperLocal.h
--- a/Engine/Plugins/VirtualProduction/CaptureData/Source/CaptureDataCore/Public/OpenCVHelperLocal.h
+++ b/Engine/Plugins/VirtualProduction/CaptureData/Source/CaptureDataCore/Public/OpenCVHelperLocal.h
@@ -3,8 +3,10 @@
 #pragma once
 #include "CoreMinimal.h"
 
+#define UE_API CAPTUREDATACORE_API
+
 // NB this is a cut-down version of the FOpenCVHelper class from OpenCV runtime plugin
-class CAPTUREDATACORE_API FOpenCVHelperLocal
+class FOpenCVHelperLocal
 {
 public:
 	/** Enumeration to specify any cartesian axis in positive or negative directions */
@@ -31,22 +33,23 @@ public:
 	};
 
 	/** Converts in-place the coordinate system of the given FTransform by specifying the source axes in terms of the destination axes */
-	static void ConvertCoordinateSystem(FTransform& Transform, const EAxis DstXInSrcAxis, const EAxis DstYInSrcAxis, const EAxis DstZInSrcAxis);
+	static UE_API void ConvertCoordinateSystem(FTransform& Transform, const EAxis DstXInSrcAxis, const EAxis DstYInSrcAxis, const EAxis DstZInSrcAxis);
 
 	/** Converts in-place an FTransform in Unreal coordinates to OpenCV coordinates */
-	static void ConvertUnrealToOpenCV(FTransform& Transform);
+	static UE_API void ConvertUnrealToOpenCV(FTransform& Transform);
 
 	/** Converts in-place an FTransform in OpenCV coordinates to Unreal coordinates */
-	static void ConvertOpenCVToUnreal(FTransform& Transform);
+	static UE_API void ConvertOpenCVToUnreal(FTransform& Transform);
 
 	/** Converts an FVector in Unreal coordinates to OpenCV coordinates */
-	static FVector ConvertUnrealToOpenCV(const FVector& Transform);
+	static UE_API FVector ConvertUnrealToOpenCV(const FVector& Transform);
 
 	/** Converts an FTransform in OpenCV coordinates to Unreal coordinates */
-	static FVector ConvertOpenCVToUnreal(const FVector& Transform);
+	static UE_API FVector ConvertOpenCVToUnreal(const FVector& Transform);
 
 	/** Converts a FMatrix and FVector in OpenCV coordinates to FRotator and FVector in Unreal coordinates */
-	static void ConvertOpenCVToUnreal(const FMatrix& InRotationOpenCV, const FVector& InTranslationOpenCV, FRotator& OutRotatorUE, FVector& OutTranslationUE);
+	static UE_API void ConvertOpenCVToUnreal(const FMatrix& InRotationOpenCV, const FVector& InTranslationOpenCV, FRotator& OutRotatorUE, FVector& OutTranslationUE);
 
 };
 
+#undef UE_API
diff --git a/Engine/Plugins/VirtualProduction/CaptureData/Source/CaptureDataEditor/Public/Customizations/CaptureDataCustomizations.h b/Engine/Plugins/VirtualProduction/CaptureData/Source/CaptureDataEditor/Public/Customizations/CaptureDataCustomizations.h
--- a/Engine/Plugins/VirtualProduction/CaptureData/Source/CaptureDataEditor/Public/Customizations/CaptureDataCustomizations.h
+++ b/Engine/Plugins/VirtualProduction/CaptureData/Source/CaptureDataEditor/Public/Customizations/CaptureDataCustomizations.h
@@ -6,13 +6,17 @@
 #include "UObject/Object.h"
 #include "IDetailCustomization.h"
 
-class CAPTUREDATAEDITOR_API FFootageCaptureDataCustomization
+#define UE_API CAPTUREDATAEDITOR_API
+
+class FFootageCaptureDataCustomization
 	: public IDetailCustomization
 {
 public:
-	static TSharedRef<IDetailCustomization> MakeInstance();
+	static UE_API TSharedRef<IDetailCustomization> MakeInstance();
 
 	//~Begin IDetailCustomization interface
-	virtual void CustomizeDetails(IDetailLayoutBuilder& InDetailBuilder) override;
+	UE_API virtual void CustomizeDetails(IDetailLayoutBuilder& InDetailBuilder) override;
 	//~End IDetailCustomization interface
-};
\ No newline at end of file
+};
+
+#undef UE_API
diff --git a/Engine/Plugins/VirtualProduction/CaptureData/Source/CaptureDataEditor/Public/FrameRangeArrayBuilder.h b/Engine/Plugins/VirtualProduction/CaptureData/Source/CaptureDataEditor/Public/FrameRangeArrayBuilder.h
--- a/Engine/Plugins/VirtualProduction/CaptureData/Source/CaptureDataEditor/Public/FrameRangeArrayBuilder.h
+++ b/Engine/Plugins/VirtualProduction/CaptureData/Source/CaptureDataEditor/Public/FrameRangeArrayBuilder.h
@@ -5,18 +5,22 @@
 #include "FrameRange.h"
 #include "PropertyCustomizationHelpers.h"
 
-class CAPTUREDATAEDITOR_API FFrameRangeArrayBuilder : public FDetailArrayBuilder
+#define UE_API CAPTUREDATAEDITOR_API
+
+class FFrameRangeArrayBuilder : public FDetailArrayBuilder
 {
 public:
 
 	DECLARE_DELEGATE_RetVal(FFrameNumber, FOnGetCurrentFrame);
 
-	FFrameRangeArrayBuilder(TSharedRef<IPropertyHandle> InBaseProperty, TArray<FFrameRange>& InOutFrameRange, FOnGetCurrentFrame* InOnGetCurrentFrameDelegate = nullptr);
+	UE_API FFrameRangeArrayBuilder(TSharedRef<IPropertyHandle> InBaseProperty, TArray<FFrameRange>& InOutFrameRange, FOnGetCurrentFrame* InOnGetCurrentFrameDelegate = nullptr);
 
-	virtual void GenerateChildContent(IDetailChildrenBuilder& InOutChildrenBuilder) override;
+	UE_API virtual void GenerateChildContent(IDetailChildrenBuilder& InOutChildrenBuilder) override;
 
 private:
 
 	TArray<FFrameRange>& FrameRange;
 	FOnGetCurrentFrame* OnGetCurrentFrameDelegate = nullptr;
 };
+
+#undef UE_API
diff --git a/Engine/Plugins/VirtualProduction/CaptureData/Source/CaptureDataUtils/Public/ImageSequenceTimecodeUtils.h b/Engine/Plugins/VirtualProduction/CaptureData/Source/CaptureDataUtils/Public/ImageSequenceTimecodeUtils.h
--- a/Engine/Plugins/VirtualProduction/CaptureData/Source/CaptureDataUtils/Public/ImageSequenceTimecodeUtils.h
+++ b/Engine/Plugins/VirtualProduction/CaptureData/Source/CaptureDataUtils/Public/ImageSequenceTimecodeUtils.h
@@ -7,36 +7,40 @@
 
 #include "ImageSequenceTimecodeUtils.generated.h"
 
-UCLASS(BlueprintType, Blueprintable)
-class CAPTUREDATAUTILS_API UImageSequenceTimecodeUtils
+#define UE_API CAPTUREDATAUTILS_API
+
+UCLASS(MinimalAPI, BlueprintType, Blueprintable)
+class UImageSequenceTimecodeUtils
 	: public UObject
 {
 	GENERATED_BODY()
 
 public:
 
-	static const FName TimecodeTagName;
-	static const FName TimecodeRateTagName;
+	static UE_API const FName TimecodeTagName;
+	static UE_API const FName TimecodeRateTagName;
 
 	UFUNCTION(BlueprintCallable, Category = "ImageSequence")
-	static void SetTimecodeInfo(const FTimecode& InTimecode, const FFrameRate& InFrameRate, UImgMediaSource* InImageSequence);
+	static UE_API void SetTimecodeInfo(const FTimecode& InTimecode, const FFrameRate& InFrameRate, UImgMediaSource* InImageSequence);
 
 	UFUNCTION(BlueprintCallable, Category = "ImageSequence")
-	static void SetTimecodeInfoString(const FString& InTimecode, const FString& InFrameRate, UImgMediaSource* InImageSequence);
+	static UE_API void SetTimecodeInfoString(const FString& InTimecode, const FString& InFrameRate, UImgMediaSource* InImageSequence);
 
 	UFUNCTION(BlueprintCallable, Category = "ImageSequence")
-	static FTimecode GetTimecode(UImgMediaSource* InImageSequence);
+	static UE_API FTimecode GetTimecode(UImgMediaSource* InImageSequence);
 
 	UFUNCTION(BlueprintCallable, Category = "ImageSequence")
-	static FFrameRate GetFrameRate(UImgMediaSource* InImageSequence);
+	static UE_API FFrameRate GetFrameRate(UImgMediaSource* InImageSequence);
 
 	UFUNCTION(BlueprintCallable, Category = "ImageSequence")
-	static FString GetTimecodeString(UImgMediaSource* InImageSequence);
+	static UE_API FString GetTimecodeString(UImgMediaSource* InImageSequence);
 
 	UFUNCTION(BlueprintCallable, Category = "ImageSequence")
-	static FString GetFrameRateString(UImgMediaSource* InImageSequence);
+	static UE_API FString GetFrameRateString(UImgMediaSource* InImageSequence);
 
-	static bool IsValidTimecodeInfo(const FTimecode& InTimecode, const FFrameRate& InTimecodeRate);
-	static bool IsValidTimecode(const FTimecode& InTimecode);
-	static bool IsValidFrameRate(const FFrameRate& InTimecodeRate);
+	static UE_API bool IsValidTimecodeInfo(const FTimecode& InTimecode, const FFrameRate& InTimecodeRate);
+	static UE_API bool IsValidTimecode(const FTimecode& InTimecode);
+	static UE_API bool IsValidFrameRate(const FFrameRate& InTimecodeRate);
 };
+
+#undef UE_API
diff --git a/Engine/Plugins/VirtualProduction/CaptureData/Source/CaptureDataUtils/Public/SoundWaveTimecodeUtils.h b/Engine/Plugins/VirtualProduction/CaptureData/Source/CaptureDataUtils/Public/SoundWaveTimecodeUtils.h
--- a/Engine/Plugins/VirtualProduction/CaptureData/Source/CaptureDataUtils/Public/SoundWaveTimecodeUtils.h
+++ b/Engine/Plugins/VirtualProduction/CaptureData/Source/CaptureDataUtils/Public/SoundWaveTimecodeUtils.h
@@ -9,19 +9,23 @@
 
 #include "SoundWaveTimecodeUtils.generated.h"
 
-UCLASS(BlueprintType, Blueprintable)
-class CAPTUREDATAUTILS_API USoundWaveTimecodeUtils : public UObject
+#define UE_API CAPTUREDATAUTILS_API
+
+UCLASS(MinimalAPI, BlueprintType, Blueprintable)
+class USoundWaveTimecodeUtils : public UObject
 {
 	GENERATED_BODY()
 
 public:
 
 	UFUNCTION(BlueprintCallable, Category = "Audio")
-	static void SetTimecodeInfo(const FTimecode& InTimecode, const FFrameRate& InFrameRate, USoundWave* OutSoundWave);
+	static UE_API void SetTimecodeInfo(const FTimecode& InTimecode, const FFrameRate& InFrameRate, USoundWave* OutSoundWave);
 
 	UFUNCTION(BlueprintCallable, Category = "Audio")
-	static FTimecode GetTimecode(const USoundWave* InSoundWave);
+	static UE_API FTimecode GetTimecode(const USoundWave* InSoundWave);
 
 	UFUNCTION(BlueprintCallable, Category = "Audio")
-	static FFrameRate GetFrameRate(const USoundWave* InSoundWave);
+	static UE_API FFrameRate GetFrameRate(const USoundWave* InSoundWave);
 };
+
+#undef UE_API
diff --git a/Engine/Plugins/VirtualProduction/CaptureManager/CaptureManagerCore/Source/CaptureManagerStyle/Public/CaptureManagerStyle.h b/Engine/Plugins/VirtualProduction/CaptureManager/CaptureManagerCore/Source/CaptureManagerStyle/Public/CaptureManagerStyle.h
--- a/Engine/Plugins/VirtualProduction/CaptureManager/CaptureManagerCore/Source/CaptureManagerStyle/Public/CaptureManagerStyle.h
+++ b/Engine/Plugins/VirtualProduction/CaptureManager/CaptureManagerCore/Source/CaptureManagerStyle/Public/CaptureManagerStyle.h
@@ -5,19 +5,23 @@
 #include "CoreMinimal.h"
 #include "Styling/SlateStyle.h"
 
+#define UE_API CAPTUREMANAGERSTYLE_API
+
 /** A style set for Capture Manager */
-class CAPTUREMANAGERSTYLE_API FCaptureManagerStyle : public FSlateStyleSet
+class FCaptureManagerStyle : public FSlateStyleSet
 {
 public:
 
-	virtual const FName& GetStyleSetName() const override;
-	static const FCaptureManagerStyle& Get();
+	UE_API virtual const FName& GetStyleSetName() const override;
+	static UE_API const FCaptureManagerStyle& Get();
 
-	static void ReloadTextures();
+	static UE_API void ReloadTextures();
 
 private:
 
-	FCaptureManagerStyle();
+	UE_API FCaptureManagerStyle();
+
+	static UE_API FName StyleName;
+};
 
-	static FName StyleName;
-};
\ No newline at end of file
+#undef UE_API
diff --git a/Engine/Plugins/VirtualProduction/CaptureManager/CaptureManagerCore/Source/CaptureManagerTakeMetadata/Public/CaptureManagerTakeMetadata.h b/Engine/Plugins/VirtualProduction/CaptureManager/CaptureManagerCore/Source/CaptureManagerTakeMetadata/Public/CaptureManagerTakeMetadata.h
--- a/Engine/Plugins/VirtualProduction/CaptureManager/CaptureManagerCore/Source/CaptureManagerTakeMetadata/Public/CaptureManagerTakeMetadata.h
+++ b/Engine/Plugins/VirtualProduction/CaptureManager/CaptureManagerCore/Source/CaptureManagerTakeMetadata/Public/CaptureManagerTakeMetadata.h
@@ -12,8 +12,10 @@
 #include "Math/MathFwd.h"
 #include "ImageCore.h"
 
+#define UE_API CAPTUREMANAGERTAKEMETADATA_API
+
 /** Class for constructing take thumbnails and accessing thumbnail data. */
-class CAPTUREMANAGERTAKEMETADATA_API FTakeThumbnailData
+class FTakeThumbnailData
 {
 public:
 
@@ -26,34 +28,34 @@ public:
 	};
 
 	/** Default thumbnail constructor. */
-	FTakeThumbnailData();
+	UE_API FTakeThumbnailData();
 
 	/** Construct thumbnail from file path. */
-	FTakeThumbnailData(FString InImagePath);
+	UE_API FTakeThumbnailData(FString InImagePath);
 
 	/** Construct thumbnail from file data array. */
-	FTakeThumbnailData(TArray<uint8> InCompressedImageData);
+	UE_API FTakeThumbnailData(TArray<uint8> InCompressedImageData);
 
 	/** Construct thumbnail from raw data array. */
-	FTakeThumbnailData(TArray<FColor> InDecompressedImageData, uint32 InWidth, uint32 InHeight, ERawImageFormat::Type InFormat);
+	UE_API FTakeThumbnailData(TArray<FColor> InDecompressedImageData, uint32 InWidth, uint32 InHeight, ERawImageFormat::Type InFormat);
 
 	/** Construct thumbnail from file path. */
-	void operator=(FString InImagePath);
+	UE_API void operator=(FString InImagePath);
 
 	/** Construct thumbnail from file data array. */
-	void operator=(TArray<uint8> InCompressedImageData);
+	UE_API void operator=(TArray<uint8> InCompressedImageData);
 
 	/** Assign thumbnail from raw data array. */
-	void operator=(FRawImage InRawImage);
+	UE_API void operator=(FRawImage InRawImage);
 
 	/** Get thumbnail. */
-	TOptional<TArray<uint8>> GetThumbnailData() const;
+	UE_API TOptional<TArray<uint8>> GetThumbnailData() const;
 	
 	/** Get thumbnail file path. */
-	TOptional<FString> GetThumbnailPath() const;
+	UE_API TOptional<FString> GetThumbnailPath() const;
 
 	/** Get decompressed image. */
-	TOptional<FRawImage> GetRawImage() const;
+	UE_API TOptional<FRawImage> GetRawImage() const;
 
 private:
 
@@ -62,11 +64,11 @@ private:
 };
 
 /** Data assosciated with a take. */
-class CAPTUREMANAGERTAKEMETADATA_API FTakeMetadata
+class FTakeMetadata
 {
 public:
 	/** File extension for take metadata files. */
-	static const FString FileExtension;
+	static UE_API const FString FileExtension;
 
 	/** Version of the take metadata schema. */
 	struct FSchemaVersion
@@ -219,7 +221,7 @@ public:
 	TArray<FAudio> Audio;
 };
 
-struct CAPTUREMANAGERTAKEMETADATA_API FTakeMetadataParserError
+struct FTakeMetadataParserError
 {
 	enum EOrigin
 	{
@@ -232,25 +234,25 @@ struct CAPTUREMANAGERTAKEMETADATA_API FTakeMetadataParserError
 	FText Message;
 };
 
-struct CAPTUREMANAGERTAKEMETADATA_API FTakeMetadataSerializerError
+struct FTakeMetadataSerializerError
 {
 	FText Message;
 };
 
-class CAPTUREMANAGERTAKEMETADATA_API FTakeMetadataParser
+class FTakeMetadataParser
 {
 public:
-	FTakeMetadataParser();
-	~FTakeMetadataParser();
+	UE_API FTakeMetadataParser();
+	UE_API ~FTakeMetadataParser();
 
 	FTakeMetadataParser(const FTakeMetadataParser& other) = delete;
-	FTakeMetadataParser(FTakeMetadataParser&& other);
+	UE_API FTakeMetadataParser(FTakeMetadataParser&& other);
 
 	FTakeMetadataParser& operator=(const FTakeMetadataParser& other) = delete;
-	FTakeMetadataParser& operator=(FTakeMetadataParser&& other);
+	UE_API FTakeMetadataParser& operator=(FTakeMetadataParser&& other);
 
 	/** Tries to parse file into a take metadata object. */
-	TValueOrError<FTakeMetadata, FTakeMetadataParserError> Parse(const FString& InJsonFile);
+	UE_API TValueOrError<FTakeMetadata, FTakeMetadataParserError> Parse(const FString& InJsonFile);
 
 private:
 	struct FImpl;
@@ -259,20 +261,21 @@ private:
 };
 
 /** Utility functions for handling paths within take metadata files. */
-class CAPTUREMANAGERTAKEMETADATA_API FTakeMetadataPathUtils
+class FTakeMetadataPathUtils
 {
 public:
 	/** Detects path type from the input path (Folder or File). */
-	static FTakeMetadata::FVideo::EPathType DetectPathType(const FString& InPath);
+	static UE_API FTakeMetadata::FVideo::EPathType DetectPathType(const FString& InPath);
 
 	/** Checks the path matches the path type. */
-	static bool ValidatePathType(const FString& InPath, FTakeMetadata::FVideo::EPathType InPathType);
+	static UE_API bool ValidatePathType(const FString& InPath, FTakeMetadata::FVideo::EPathType InPathType);
 
 	/** Converts path type to a string. */
-	static FString PathTypeToString(FTakeMetadata::FVideo::EPathType InPathType);
+	static UE_API FString PathTypeToString(FTakeMetadata::FVideo::EPathType InPathType);
 };
 
 /** Writes a take metadata object to the specified path. */
 CAPTUREMANAGERTAKEMETADATA_API TOptional<FTakeMetadataSerializerError> SerializeTakeMetadata(const FString& InFilePath, const FTakeMetadata& InMetadata);
 
 
+#undef UE_API
diff --git a/Engine/Plugins/VirtualProduction/CaptureManager/CaptureManagerCore/Source/CaptureProtocolStack/Public/Control/ControlMessenger.h b/Engine/Plugins/VirtualProduction/CaptureManager/CaptureManagerCore/Source/CaptureProtocolStack/Public/Control/ControlMessenger.h
--- a/Engine/Plugins/VirtualProduction/CaptureManager/CaptureManagerCore/Source/CaptureProtocolStack/Public/Control/ControlMessenger.h
+++ b/Engine/Plugins/VirtualProduction/CaptureManager/CaptureManagerCore/Source/CaptureProtocolStack/Public/Control/ControlMessenger.h
@@ -19,6 +19,8 @@
 
 #include "Async/Async.h"
 
+#define UE_API CAPTUREPROTOCOLSTACK_API
+
 namespace UE::CaptureManager
 {
 
@@ -36,7 +38,7 @@ private:
 	std::atomic<uint16> Counter;
 };
 
-class CAPTUREPROTOCOLSTACK_API FControlMessenger final
+class FControlMessenger final
 {
 public:
 
@@ -47,19 +49,19 @@ public:
 
 	static constexpr int32 ResponseWaitTime = 3; // Seconds
 	static constexpr int32 KeepAliveInterval = 5; // Seconds
-	static const TCHAR* HandshakeSessionId;
+	static UE_API const TCHAR* HandshakeSessionId;
 
-	FControlMessenger();
-	~FControlMessenger();
+	UE_API FControlMessenger();
+	UE_API ~FControlMessenger();
 
-	void RegisterUpdateHandler(FString InAddressPath, FControlUpdate::FOnUpdateMessage InUpdateHandler);
-	void RegisterDisconnectHandler(FOnDisconnect InOnDisconnectHandler);
+	UE_API void RegisterUpdateHandler(FString InAddressPath, FControlUpdate::FOnUpdateMessage InUpdateHandler);
+	UE_API void RegisterDisconnectHandler(FOnDisconnect InOnDisconnectHandler);
 
-	TProtocolResult<void> Start(const FString& InServerIp, const uint16 InServerPort);
-	void Stop();
+	UE_API TProtocolResult<void> Start(const FString& InServerIp, const uint16 InServerPort);
+	UE_API void Stop();
 
-	TProtocolResult<void> StartSession();
-	TProtocolResult<FGetServerInformationResponse> GetServerInformation();
+	UE_API TProtocolResult<void> StartSession();
+	UE_API TProtocolResult<FGetServerInformationResponse> GetServerInformation();
 
 	template<class RequestType>
 	TProtocolResult<typename RequestType::ResponseType> SendRequest(RequestType InRequest)
@@ -166,18 +168,18 @@ private:
 		TPromise<TProtocolResult<FControlMessage>> Promise;
 	};
 
-	void SendPacket(FControlPacket InPacket);
-	void KeepAlive();
-	void MessageHandler(FControlPacket InPacket);
-	void CommunicationStoppedHandler();
+	UE_API void SendPacket(FControlPacket InPacket);
+	UE_API void KeepAlive();
+	UE_API void MessageHandler(FControlPacket InPacket);
+	UE_API void CommunicationStoppedHandler();
 
-	uint32 GenerateTransactionId() const;
-	uint64 GetTimestamp() const;
+	UE_API uint32 GenerateTransactionId() const;
+	UE_API uint64 GetTimestamp() const;
 
-	void StartKeepAliveTimer();
-	void StopKeepAliveTimer();
+	UE_API void StartKeepAliveTimer();
+	UE_API void StopKeepAliveTimer();
 
-	void OnAsyncRequestProcess(FAsyncRequestDelegate InAsyncDelegate);
+	UE_API void OnAsyncRequestProcess(FAsyncRequestDelegate InAsyncDelegate);
 	
 	FControlCommunication Communication;
 
@@ -199,4 +201,6 @@ private:
 	FRandomStream RandomStream;
 };
 
-}
\ No newline at end of file
+}
+
+#undef UE_API
diff --git a/Engine/Plugins/VirtualProduction/CaptureManager/CaptureManagerCore/Source/CaptureProtocolStack/Public/Control/Messages/ControlJsonUtilities.h b/Engine/Plugins/VirtualProduction/CaptureManager/CaptureManagerCore/Source/CaptureProtocolStack/Public/Control/Messages/ControlJsonUtilities.h
--- a/Engine/Plugins/VirtualProduction/CaptureManager/CaptureManagerCore/Source/CaptureProtocolStack/Public/Control/Messages/ControlJsonUtilities.h
+++ b/Engine/Plugins/VirtualProduction/CaptureManager/CaptureManagerCore/Source/CaptureProtocolStack/Public/Control/Messages/ControlJsonUtilities.h
@@ -10,16 +10,18 @@
 
 #include "Policies/CondensedJsonPrintPolicy.h"
 
+#define UE_API CAPTUREPROTOCOLSTACK_API
+
 #define CHECK_PARSE CPS_CHECK_VOID_RESULT
 
 namespace UE::CaptureManager
 {
 
-class CAPTUREPROTOCOLSTACK_API FJsonUtility final
+class FJsonUtility final
 {
 public:
-	static bool CreateJsonFromUTF8Data(const TArray<uint8>& InData, TSharedPtr<FJsonObject>& OutObject);
-	static bool CreateUTF8DataFromJson(TSharedPtr<FJsonObject> InObject, TArray<uint8>& OutData);
+	static UE_API bool CreateJsonFromUTF8Data(const TArray<uint8>& InData, TSharedPtr<FJsonObject>& OutObject);
+	static UE_API bool CreateUTF8DataFromJson(TSharedPtr<FJsonObject> InObject, TArray<uint8>& OutData);
 
 	template <typename T>
 	static TProtocolResult<void> ParseNumber(const TSharedPtr<FJsonObject>& InBody, const FString& InFieldName, T& OutFieldValue)
@@ -31,10 +33,12 @@ public:
 
 		return ResultOk;
 	}
-	static TProtocolResult<void> ParseString(const TSharedPtr<FJsonObject>& InBody, const FString& InFieldName, FString& OutFieldValue);
-	static TProtocolResult<void> ParseBool(const TSharedPtr<FJsonObject>& InBody, const FString& InFieldName, bool& OutFieldValue);
-	static TProtocolResult<void> ParseObject(const TSharedPtr<FJsonObject>& InBody, const FString& InFieldName, const TSharedPtr<FJsonObject>*& OutFieldValue);
-	static TProtocolResult<void> ParseArray(const TSharedPtr<FJsonObject>& InBody, const FString& InFieldName, const TArray<TSharedPtr<FJsonValue>>*& OutFieldValue);
+	static UE_API TProtocolResult<void> ParseString(const TSharedPtr<FJsonObject>& InBody, const FString& InFieldName, FString& OutFieldValue);
+	static UE_API TProtocolResult<void> ParseBool(const TSharedPtr<FJsonObject>& InBody, const FString& InFieldName, bool& OutFieldValue);
+	static UE_API TProtocolResult<void> ParseObject(const TSharedPtr<FJsonObject>& InBody, const FString& InFieldName, const TSharedPtr<FJsonObject>*& OutFieldValue);
+	static UE_API TProtocolResult<void> ParseArray(const TSharedPtr<FJsonObject>& InBody, const FString& InFieldName, const TArray<TSharedPtr<FJsonValue>>*& OutFieldValue);
 };
 
-}
\ No newline at end of file
+}
+
+#undef UE_API
diff --git a/Engine/Plugins/VirtualProduction/CaptureManager/CaptureManagerCore/Source/CaptureProtocolStack/Public/Control/Messages/ControlMessage.h b/Engine/Plugins/VirtualProduction/CaptureManager/CaptureManagerCore/Source/CaptureProtocolStack/Public/Control/Messages/ControlMessage.h
--- a/Engine/Plugins/VirtualProduction/CaptureManager/CaptureManagerCore/Source/CaptureProtocolStack/Public/Control/Messages/ControlMessage.h
+++ b/Engine/Plugins/VirtualProduction/CaptureManager/CaptureManagerCore/Source/CaptureProtocolStack/Public/Control/Messages/ControlMessage.h
@@ -4,10 +4,12 @@
 
 #include "Control/Communication/ControlPacket.h"
 
+#define UE_API CAPTUREPROTOCOLSTACK_API
+
 namespace UE::CaptureManager
 {
 
-class CAPTUREPROTOCOLSTACK_API FControlMessage final
+class FControlMessage final
 {
 public:
 
@@ -20,7 +22,7 @@ public:
 		Invalid
 	};
 
-	FControlMessage(FString InAddressPath, EType InType, TSharedPtr<FJsonObject> InBody);
+	UE_API FControlMessage(FString InAddressPath, EType InType, TSharedPtr<FJsonObject> InBody);
 
 	FControlMessage(const FControlMessage& InOther) = default;
 	FControlMessage(FControlMessage&& InOther) = default;
@@ -28,23 +30,23 @@ public:
 	FControlMessage& operator=(const FControlMessage& InOther) = default;
 	FControlMessage& operator=(FControlMessage&& InOther) = default;
 
-	static TProtocolResult<FControlMessage> Deserialize(const FControlPacket& InPacket);
-	static TProtocolResult<FControlPacket> Serialize(const FControlMessage& InMessage);
+	static UE_API TProtocolResult<FControlMessage> Deserialize(const FControlPacket& InPacket);
+	static UE_API TProtocolResult<FControlPacket> Serialize(const FControlMessage& InMessage);
 
-	void SetSessionId(FString InSessionId);
-	void SetTransactionId(uint32 InTrasactionId);
-	void SetTimestamp(uint64 InTimestamp);
+	UE_API void SetSessionId(FString InSessionId);
+	UE_API void SetTransactionId(uint32 InTrasactionId);
+	UE_API void SetTimestamp(uint64 InTimestamp);
 
-	const FString& GetSessionId() const;
-	const FString& GetAddressPath() const;
-	uint32 GetTransactionId() const;
-	uint64 GetTimestamp() const;
-	EType GetType() const;
-	const TSharedPtr<FJsonObject>& GetBody() const;
-	TSharedPtr<FJsonObject>& GetBody();
+	UE_API const FString& GetSessionId() const;
+	UE_API const FString& GetAddressPath() const;
+	UE_API uint32 GetTransactionId() const;
+	UE_API uint64 GetTimestamp() const;
+	UE_API EType GetType() const;
+	UE_API const TSharedPtr<FJsonObject>& GetBody() const;
+	UE_API TSharedPtr<FJsonObject>& GetBody();
 
-	const FString& GetErrorName() const;
-	const FString& GetErrorDescription() const;
+	UE_API const FString& GetErrorName() const;
+	UE_API const FString& GetErrorDescription() const;
 
 private:
 
@@ -54,7 +56,7 @@ private:
 		FString Description = TEXT("");
 	};
 
-	FControlMessage(FString InSessionId,
+	UE_API FControlMessage(FString InSessionId,
 					FString InAddressPath,
 					uint32 InTransactionId,
 					uint64 InTimestamp,
@@ -62,8 +64,8 @@ private:
 					TSharedPtr<FJsonObject> InBody,
 					FErrorResponse InError);
 
-	static EType DeserializeType(FString InMessageTypeStr);
-	static FString SerializeType(EType InMessageType);
+	static UE_API EType DeserializeType(FString InMessageTypeStr);
+	static UE_API FString SerializeType(EType InMessageType);
 
 	FString SessionId;
 	FString AddressPath;
@@ -74,4 +76,6 @@ private:
 	FErrorResponse Error;
 };
 
-}
\ No newline at end of file
+}
+
+#undef UE_API
diff --git a/Engine/Plugins/VirtualProduction/CaptureManager/CaptureManagerCore/Source/CaptureProtocolStack/Public/Control/Messages/ControlRequest.h b/Engine/Plugins/VirtualProduction/CaptureManager/CaptureManagerCore/Source/CaptureProtocolStack/Public/Control/Messages/ControlRequest.h
--- a/Engine/Plugins/VirtualProduction/CaptureManager/CaptureManagerCore/Source/CaptureProtocolStack/Public/Control/Messages/ControlRequest.h
+++ b/Engine/Plugins/VirtualProduction/CaptureManager/CaptureManagerCore/Source/CaptureProtocolStack/Public/Control/Messages/ControlRequest.h
@@ -5,101 +5,103 @@
 #include "ControlResponse.h"
 #include "Misc/Optional.h"
 
+#define UE_API CAPTUREPROTOCOLSTACK_API
+
 namespace UE::CaptureManager
 {
 
-class CAPTUREPROTOCOLSTACK_API FControlRequest
+class FControlRequest
 {
 public:
 
-	FControlRequest(FString InAddressPath);
+	UE_API FControlRequest(FString InAddressPath);
 	virtual ~FControlRequest() = default;
 
-	const FString& GetAddressPath() const;
-	virtual TSharedPtr<FJsonObject> GetBody() const;
+	UE_API const FString& GetAddressPath() const;
+	UE_API virtual TSharedPtr<FJsonObject> GetBody() const;
 
 private:
 
 	FString AddressPath;
 };
 
-class CAPTUREPROTOCOLSTACK_API FKeepAliveRequest final : public FControlRequest
+class FKeepAliveRequest final : public FControlRequest
 {
 public:
 
 	using ResponseType = FKeepAliveResponse;
 
-	FKeepAliveRequest();
+	UE_API FKeepAliveRequest();
 };
 
-class CAPTUREPROTOCOLSTACK_API FStartSessionRequest final : public FControlRequest
+class FStartSessionRequest final : public FControlRequest
 {
 public:
 
 	using ResponseType = FStartSessionResponse;
 
-	FStartSessionRequest();
+	UE_API FStartSessionRequest();
 };
 
-class CAPTUREPROTOCOLSTACK_API FStopSessionRequest final : public FControlRequest
+class FStopSessionRequest final : public FControlRequest
 {
 public:
 
 	using ResponseType = FStopSessionResponse;
 
-	FStopSessionRequest();
+	UE_API FStopSessionRequest();
 };
 
 
-class CAPTUREPROTOCOLSTACK_API FGetServerInformationRequest final : public FControlRequest
+class FGetServerInformationRequest final : public FControlRequest
 {
 public:
 
 	using ResponseType = FGetServerInformationResponse;
 
-	FGetServerInformationRequest();
+	UE_API FGetServerInformationRequest();
 };
 
-class CAPTUREPROTOCOLSTACK_API FSubscribeRequest final : public FControlRequest
+class FSubscribeRequest final : public FControlRequest
 {
 public:
 
 	using ResponseType = FSubscribeResponse;
 
-	FSubscribeRequest();
+	UE_API FSubscribeRequest();
 };
 
-class CAPTUREPROTOCOLSTACK_API FUnsubscribeRequest final : public FControlRequest
+class FUnsubscribeRequest final : public FControlRequest
 {
 public:
 
 	using ResponseType = FUnsubscribeResponse;
 
-	FUnsubscribeRequest();
+	UE_API FUnsubscribeRequest();
 };
 
-class CAPTUREPROTOCOLSTACK_API FGetStateRequest final : public FControlRequest
+class FGetStateRequest final : public FControlRequest
 {
 public:
 
 	using ResponseType = FGetStateResponse;
 
-	FGetStateRequest();
+	UE_API FGetStateRequest();
 };
 
-class CAPTUREPROTOCOLSTACK_API FStartRecordingTakeRequest final : public FControlRequest
+class FStartRecordingTakeRequest final : public FControlRequest
 {
 public:
 
 	using ResponseType = FStartRecordingTakeResponse;
 
-	FStartRecordingTakeRequest(FString InSlateName,
+	UE_API FStartRecordingTakeRequest(FString InSlateName,
 							   uint16 InTakeNumber,
 							   TOptional<FString> InSubject = TOptional<FString>(),
 							   TOptional<FString> InScenario = TOptional<FString>(),
 							   TOptional<TArray<FString>> InTags = TOptional<TArray<FString>>());
 
-	virtual TSharedPtr<FJsonObject> GetBody() const override;
+	UE_API virtual TSharedPtr<FJsonObject> GetBody() const override;
 
 private:
 
@@ -110,57 +112,57 @@ private:
 	TOptional<TArray<FString>> Tags;
 };
 
-class CAPTUREPROTOCOLSTACK_API FStopRecordingTakeRequest final : public FControlRequest
+class FStopRecordingTakeRequest final : public FControlRequest
 {
 public:
 
 	using ResponseType = FStopRecordingTakeResponse;
 
-	FStopRecordingTakeRequest();
+	UE_API FStopRecordingTakeRequest();
 };
 
-class CAPTUREPROTOCOLSTACK_API FAbortRecordingTakeRequest final : public FControlRequest
+class FAbortRecordingTakeRequest final : public FControlRequest
 {
 public:
 
 	using ResponseType = FAbortRecordingTakeResponse;
 
-	FAbortRecordingTakeRequest();
+	UE_API FAbortRecordingTakeRequest();
 };
 
-class CAPTUREPROTOCOLSTACK_API FGetTakeListRequest final : public FControlRequest
+class FGetTakeListRequest final : public FControlRequest
 {
 public:
 
 	using ResponseType = FGetTakeListResponse;
 
-	FGetTakeListRequest();
+	UE_API FGetTakeListRequest();
 };
 
-class CAPTUREPROTOCOLSTACK_API FGetTakeMetadataRequest final : public FControlRequest
+class FGetTakeMetadataRequest final : public FControlRequest
 {
 public:
 
 	using ResponseType = FGetTakeMetadataResponse;
 
-	FGetTakeMetadataRequest(TArray<FString> InNames);
+	UE_API FGetTakeMetadataRequest(TArray<FString> InNames);
 
-	virtual TSharedPtr<FJsonObject> GetBody() const override;
+	UE_API virtual TSharedPtr<FJsonObject> GetBody() const override;
 
 private:
 
 	TArray<FString> Names;
 };
 
-class CAPTUREPROTOCOLSTACK_API FGetStreamingSubjectsRequest final : public FControlRequest
+class FGetStreamingSubjectsRequest final : public FControlRequest
 {
 public:
 	using ResponseType = FGetStreamingSubjectsResponse;
 
-	FGetStreamingSubjectsRequest();
+	UE_API FGetStreamingSubjectsRequest();
 };
 
-class CAPTUREPROTOCOLSTACK_API FStartStreamingRequest final : public FControlRequest
+class FStartStreamingRequest final : public FControlRequest
 {
 public:
 
@@ -185,9 +187,9 @@ public:
 	
 	using ResponseType = FStartStreamingResponse;
 
-	FStartStreamingRequest(uint16 InStreamPort, TArray<FSubject> InSubjects);
+	UE_API FStartStreamingRequest(uint16 InStreamPort, TArray<FSubject> InSubjects);
 
-	virtual TSharedPtr<FJsonObject> GetBody() const override;
+	UE_API virtual TSharedPtr<FJsonObject> GetBody() const override;
 	
 private:
 	
@@ -195,18 +197,20 @@ private:
 	TArray<FSubject> Subjects;
 };
 
-class CAPTUREPROTOCOLSTACK_API FStopStreamingRequest final : public FControlRequest
+class FStopStreamingRequest final : public FControlRequest
 {
 public:
 
 	using ResponseType = FStopStreamingResponse;
 
-	FStopStreamingRequest(TOptional<TArray<FString>> InSubjectIds = TOptional<TArray<FString>>());
+	UE_API FStopStreamingRequest(TOptional<TArray<FString>> InSubjectIds = TOptional<TArray<FString>>());
 	
-	virtual TSharedPtr<FJsonObject> GetBody() const override;
+	UE_API virtual TSharedPtr<FJsonObject> GetBody() const override;
 	
 private:
 	TOptional<TArray<FString>> SubjectIds;
 };
 
-}
\ No newline at end of file
+}
+
+#undef UE_API
diff --git a/Engine/Plugins/VirtualProduction/CaptureManager/CaptureManagerCore/Source/CaptureProtocolStack/Public/Control/Messages/ControlResponse.h b/Engine/Plugins/VirtualProduction/CaptureManager/CaptureManagerCore/Source/CaptureProtocolStack/Public/Control/Messages/ControlResponse.h
--- a/Engine/Plugins/VirtualProduction/CaptureManager/CaptureManagerCore/Source/CaptureProtocolStack/Public/Control/Messages/ControlResponse.h
+++ b/Engine/Plugins/VirtualProduction/CaptureManager/CaptureManagerCore/Source/CaptureProtocolStack/Public/Control/Messages/ControlResponse.h
@@ -4,70 +4,72 @@
 
 #include "ControlJsonUtilities.h"
 
+#define UE_API CAPTUREPROTOCOLSTACK_API
+
 namespace UE::CaptureManager
 {
 
-class CAPTUREPROTOCOLSTACK_API FControlResponse
+class FControlResponse
 {
 public:
 
-	FControlResponse(FString InAddressPath);
+	UE_API FControlResponse(FString InAddressPath);
 	virtual ~FControlResponse() = default;
 
-	virtual TProtocolResult<void> Parse(TSharedPtr<FJsonObject> InBody);
+	UE_API virtual TProtocolResult<void> Parse(TSharedPtr<FJsonObject> InBody);
 
-	const FString& GetAddressPath() const;
+	UE_API const FString& GetAddressPath() const;
 
 private:
 
 	FString AddressPath;
 };
 
-class CAPTUREPROTOCOLSTACK_API FKeepAliveResponse final : public FControlResponse
+class FKeepAliveResponse final : public FControlResponse
 {
 public:
 
-	FKeepAliveResponse();
+	UE_API FKeepAliveResponse();
 };
 
-class CAPTUREPROTOCOLSTACK_API FStartSessionResponse final : public FControlResponse
+class FStartSessionResponse final : public FControlResponse
 {
 public:
 
-	FStartSessionResponse();
+	UE_API FStartSessionResponse();
 
-	virtual TProtocolResult<void> Parse(TSharedPtr<FJsonObject> InBody) override;
+	UE_API virtual TProtocolResult<void> Parse(TSharedPtr<FJsonObject> InBody) override;
 
-	const FString& GetSessionId() const;
+	UE_API const FString& GetSessionId() const;
 
 private:
 
 	FString SessionId;
 };
 
-class CAPTUREPROTOCOLSTACK_API FStopSessionResponse final : public FControlResponse
+class FStopSessionResponse final : public FControlResponse
 {
 public:
 
-	FStopSessionResponse();
+	UE_API FStopSessionResponse();
 };
 
-class CAPTUREPROTOCOLSTACK_API FGetServerInformationResponse final : public FControlResponse
+class FGetServerInformationResponse final : public FControlResponse
 {
 public:
 
-	FGetServerInformationResponse();
+	UE_API FGetServerInformationResponse();
 
-	virtual TProtocolResult<void> Parse(TSharedPtr<FJsonObject> InBody) override;
+	UE_API virtual TProtocolResult<void> Parse(TSharedPtr<FJsonObject> InBody) override;
 
-	const FString& GetId() const;
-	const FString& GetName() const;
-	const FString& GetModel() const;
-	const FString& GetPlatformName() const;
-	const FString& GetPlatformVersion() const;
-	const FString& GetSoftwareName() const;
-	const FString& GetSoftwareVersion() const;
-	uint16 GetExportPort() const;
+	UE_API const FString& GetId() const;
+	UE_API const FString& GetName() const;
+	UE_API const FString& GetModel() const;
+	UE_API const FString& GetPlatformName() const;
+	UE_API const FString& GetPlatformVersion() const;
+	UE_API const FString& GetSoftwareName() const;
+	UE_API const FString& GetSoftwareVersion() const;
+	UE_API uint16 GetExportPort() const;
 
 private:
 
@@ -81,31 +83,31 @@ private:
 	uint16 ExportPort = 0;
 };
 
-class CAPTUREPROTOCOLSTACK_API FSubscribeResponse final : public FControlResponse
+class FSubscribeResponse final : public FControlResponse
 {
 public:
 
-	FSubscribeResponse();
+	UE_API FSubscribeResponse();
 };
 
-class CAPTUREPROTOCOLSTACK_API FUnsubscribeResponse final : public FControlResponse
+class FUnsubscribeResponse final : public FControlResponse
 {
 public:
 
-	FUnsubscribeResponse();
+	UE_API FUnsubscribeResponse();
 };
 
 
-class CAPTUREPROTOCOLSTACK_API FGetStateResponse final : public FControlResponse
+class FGetStateResponse final : public FControlResponse
 {
 public:
 
-	FGetStateResponse();
+	UE_API FGetStateResponse();
 
-	virtual TProtocolResult<void> Parse(TSharedPtr<FJsonObject> InBody) override;
+	UE_API virtual TProtocolResult<void> Parse(TSharedPtr<FJsonObject> InBody) override;
 
-	bool IsRecording() const;
-	const TSharedPtr<FJsonObject>& GetPlatformState() const;
+	UE_API bool IsRecording() const;
+	UE_API const TSharedPtr<FJsonObject>& GetPlatformState() const;
 
 private:
 
@@ -113,50 +115,50 @@ private:
 	TSharedPtr<FJsonObject> PlatformState;
 };
 
-class CAPTUREPROTOCOLSTACK_API FStartRecordingTakeResponse final : public FControlResponse
+class FStartRecordingTakeResponse final : public FControlResponse
 {
 public:
 
-	FStartRecordingTakeResponse();
+	UE_API FStartRecordingTakeResponse();
 };
 
-class CAPTUREPROTOCOLSTACK_API FStopRecordingTakeResponse final : public FControlResponse
+class FStopRecordingTakeResponse final : public FControlResponse
 {
 public:
 
-	FStopRecordingTakeResponse();
+	UE_API FStopRecordingTakeResponse();
 
-	virtual TProtocolResult<void> Parse(TSharedPtr<FJsonObject> InBody) override;
+	UE_API virtual TProtocolResult<void> Parse(TSharedPtr<FJsonObject> InBody) override;
 
-	const FString& GetTakeName() const;
+	UE_API const FString& GetTakeName() const;
 private:
 
 	FString TakeName;
 };
 
-class CAPTUREPROTOCOLSTACK_API FAbortRecordingTakeResponse final : public FControlResponse
+class FAbortRecordingTakeResponse final : public FControlResponse
 {
 public:
 
-	FAbortRecordingTakeResponse();
+	UE_API FAbortRecordingTakeResponse();
 };
 
-class CAPTUREPROTOCOLSTACK_API FGetTakeListResponse final : public FControlResponse
+class FGetTakeListResponse final : public FControlResponse
 {
 public:
 
-	FGetTakeListResponse();
+	UE_API FGetTakeListResponse();
 
-	virtual TProtocolResult<void> Parse(TSharedPtr<FJsonObject> InBody) override;
+	UE_API virtual TProtocolResult<void> Parse(TSharedPtr<FJsonObject> InBody) override;
 
-	const TArray<FString>& GetNames() const;
+	UE_API const TArray<FString>& GetNames() const;
 
 private:
 
 	TArray<FString> Names;
 };
 
-class CAPTUREPROTOCOLSTACK_API FGetTakeMetadataResponse final : public FControlResponse
+class FGetTakeMetadataResponse final : public FControlResponse
 {
 public:
 
@@ -199,21 +201,21 @@ public:
 		FAudioObject Audio;
 	};
 
-	FGetTakeMetadataResponse();
+	UE_API FGetTakeMetadataResponse();
 
-	virtual TProtocolResult<void> Parse(TSharedPtr<FJsonObject> InBody) override;
+	UE_API virtual TProtocolResult<void> Parse(TSharedPtr<FJsonObject> InBody) override;
 
-	const TArray<FTakeObject>& GetTakes() const;
+	UE_API const TArray<FTakeObject>& GetTakes() const;
 
 private:
 
-	TProtocolResult<void> CreateTakeObject(const TSharedPtr<FJsonObject>& InTakeObject, FTakeObject& OutTake) const;
-	TProtocolResult<void> CreateFileObject(const TSharedPtr<FJsonObject>& InFileObject, FFileObject& OutFile) const;
+	UE_API TProtocolResult<void> CreateTakeObject(const TSharedPtr<FJsonObject>& InTakeObject, FTakeObject& OutTake) const;
+	UE_API TProtocolResult<void> CreateFileObject(const TSharedPtr<FJsonObject>& InFileObject, FFileObject& OutFile) const;
 
 	TArray<FTakeObject> Takes;
 };
 
-class CAPTUREPROTOCOLSTACK_API FGetStreamingSubjectsResponse final : public FControlResponse
+class FGetStreamingSubjectsResponse final : public FControlResponse
 {
 public:
 	
@@ -231,33 +233,35 @@ public:
 		FAnimationMetadata AnimationMetadata;
 	};
 	
-	FGetStreamingSubjectsResponse();
+	UE_API FGetStreamingSubjectsResponse();
 
-	virtual TProtocolResult<void> Parse(TSharedPtr<FJsonObject> InBody) override;
+	UE_API virtual TProtocolResult<void> Parse(TSharedPtr<FJsonObject> InBody) override;
 
-	const TArray<FSubject>& GetSubjects() const;
+	UE_API const TArray<FSubject>& GetSubjects() const;
 
 private:
 	
 	TArray<FSubject> Subjects;
 	
-	TProtocolResult<void> CreateSubject(const TSharedPtr<FJsonObject>& InSubjectObject, FSubject& OutSubject) const;
-	TProtocolResult<void> CreateAnimationMetadata(const TSharedPtr<FJsonObject>& InAnimationObject, FAnimationMetadata& OutAnimationMetadata) const;
+	UE_API TProtocolResult<void> CreateSubject(const TSharedPtr<FJsonObject>& InSubjectObject, FSubject& OutSubject) const;
+	UE_API TProtocolResult<void> CreateAnimationMetadata(const TSharedPtr<FJsonObject>& InAnimationObject, FAnimationMetadata& OutAnimationMetadata) const;
 };
 
-class CAPTUREPROTOCOLSTACK_API FStartStreamingResponse final : public FControlResponse
+class FStartStreamingResponse final : public FControlResponse
 {
 public:
 
-	FStartStreamingResponse();
+	UE_API FStartStreamingResponse();
 	
 };
 
-class CAPTUREPROTOCOLSTACK_API FStopStreamingResponse final : public FControlResponse
+class FStopStreamingResponse final : public FControlResponse
 {
 public:
 
-	FStopStreamingResponse();
+	UE_API FStopStreamingResponse();
 };
 
 }
+
+#undef UE_API
diff --git a/Engine/Plugins/VirtualProduction/CaptureManager/CaptureManagerCore/Source/CaptureProtocolStack/Public/Control/Messages/ControlUpdate.h b/Engine/Plugins/VirtualProduction/CaptureManager/CaptureManagerCore/Source/CaptureProtocolStack/Public/Control/Messages/ControlUpdate.h
--- a/Engine/Plugins/VirtualProduction/CaptureManager/CaptureManagerCore/Source/CaptureProtocolStack/Public/Control/Messages/ControlUpdate.h
+++ b/Engine/Plugins/VirtualProduction/CaptureManager/CaptureManagerCore/Source/CaptureProtocolStack/Public/Control/Messages/ControlUpdate.h
@@ -8,6 +8,8 @@
 
 #include "Delegates/Delegate.h"
 
+#define UE_API CAPTUREPROTOCOLSTACK_API
+
 namespace UE::CaptureManager
 {
 
@@ -18,94 +20,94 @@ public:
 	static TProtocolResult<TSharedRef<FControlUpdate>> Create(const FString& InAddressPath);
 };
 
-class CAPTUREPROTOCOLSTACK_API FControlUpdate
+class FControlUpdate
 {
 public:
 
 	DECLARE_DELEGATE_OneParam(FOnUpdateMessage, TSharedPtr<FControlUpdate> InUpdateMessage);
 
-	FControlUpdate(FString InAddressPath);
+	UE_API FControlUpdate(FString InAddressPath);
 	virtual ~FControlUpdate() = default;
 
-	const FString& GetAddressPath() const;
+	UE_API const FString& GetAddressPath() const;
 
-	virtual TProtocolResult<void> Parse(TSharedPtr<FJsonObject> InBody);
+	UE_API virtual TProtocolResult<void> Parse(TSharedPtr<FJsonObject> InBody);
 
 private:
 
 	FString AddressPath;
 };
 
-class CAPTUREPROTOCOLSTACK_API FSessionStopped final : public FControlUpdate
+class FSessionStopped final : public FControlUpdate
 {
 public:
 
-	FSessionStopped();
+	UE_API FSessionStopped();
 };
 
-class CAPTUREPROTOCOLSTACK_API FRecordingStatusUpdate final : public FControlUpdate
+class FRecordingStatusUpdate final : public FControlUpdate
 {
 public:
 
-	FRecordingStatusUpdate();
+	UE_API FRecordingStatusUpdate();
 
-	virtual TProtocolResult<void> Parse(TSharedPtr<FJsonObject> InBody) override;
+	UE_API virtual TProtocolResult<void> Parse(TSharedPtr<FJsonObject> InBody) override;
 
-	bool IsRecording() const;
+	UE_API bool IsRecording() const;
 
 private:
 
 	bool bIsRecording;
 };
 
-class CAPTUREPROTOCOLSTACK_API FBaseTakeUpdate : public FControlUpdate
+class FBaseTakeUpdate : public FControlUpdate
 {
 public:
 
 	using FControlUpdate::FControlUpdate;
 
-	virtual TProtocolResult<void> Parse(TSharedPtr<FJsonObject> InBody) override;
+	UE_API virtual TProtocolResult<void> Parse(TSharedPtr<FJsonObject> InBody) override;
 
-	const FString& GetTakeName() const;
+	UE_API const FString& GetTakeName() const;
 
 private:
 
 	FString TakeName;
 };
 
-class CAPTUREPROTOCOLSTACK_API FTakeAddedUpdate final : public FBaseTakeUpdate
+class FTakeAddedUpdate final : public FBaseTakeUpdate
 {
 public:
 
-	FTakeAddedUpdate();
+	UE_API FTakeAddedUpdate();
 };
 
-class CAPTUREPROTOCOLSTACK_API FTakeRemovedUpdate final : public FBaseTakeUpdate
+class FTakeRemovedUpdate final : public FBaseTakeUpdate
 {
 public:
 
-	FTakeRemovedUpdate();
+	UE_API FTakeRemovedUpdate();
 };
 
-class CAPTUREPROTOCOLSTACK_API FTakeUpdatedUpdate final : public FBaseTakeUpdate
+class FTakeUpdatedUpdate final : public FBaseTakeUpdate
 {
 public:
 
-	FTakeUpdatedUpdate();
+	UE_API FTakeUpdatedUpdate();
 };
 
 // IOS
 
-class CAPTUREPROTOCOLSTACK_API FDiskCapacityUpdate final : public FControlUpdate
+class FDiskCapacityUpdate final : public FControlUpdate
 {
 public:
 
-	FDiskCapacityUpdate();
+	UE_API FDiskCapacityUpdate();
 
-	virtual TProtocolResult<void> Parse(TSharedPtr<FJsonObject> InBody) override;
+	UE_API virtual TProtocolResult<void> Parse(TSharedPtr<FJsonObject> InBody) override;
 
-	uint64 GetTotal() const;
-	uint64 GetRemaining() const;
+	UE_API uint64 GetTotal() const;
+	UE_API uint64 GetRemaining() const;
 
 private:
 
@@ -113,22 +115,22 @@ private:
 	uint64 Remaining = 0;
 };
 
-class CAPTUREPROTOCOLSTACK_API FBatteryPercentageUpdate final : public FControlUpdate
+class FBatteryPercentageUpdate final : public FControlUpdate
 {
 public:
 
-	FBatteryPercentageUpdate();
+	UE_API FBatteryPercentageUpdate();
 
-	virtual TProtocolResult<void> Parse(TSharedPtr<FJsonObject> InBody) override;
+	UE_API virtual TProtocolResult<void> Parse(TSharedPtr<FJsonObject> InBody) override;
 
-	float GetLevel() const;
+	UE_API float GetLevel() const;
 
 private:
 
 	float Level = 0.0f;
 };
 
-class CAPTUREPROTOCOLSTACK_API FThermalStateUpdate final : public FControlUpdate
+class FThermalStateUpdate final : public FControlUpdate
 {
 public:
 
@@ -142,17 +144,19 @@ public:
 		Invalid
 	};
 
-	FThermalStateUpdate();
+	UE_API FThermalStateUpdate();
 
-	virtual TProtocolResult<void> Parse(TSharedPtr<FJsonObject> InBody) override;
+	UE_API virtual TProtocolResult<void> Parse(TSharedPtr<FJsonObject> InBody) override;
 
-	EState GetState() const;
+	UE_API EState GetState() const;
 
 private:
 
-	static EState ConvertState(const FString& InStateString);
+	static UE_API EState ConvertState(const FString& InStateString);
 
 	EState State;
 };
 
-}
\ No newline at end of file
+}
+
+#undef UE_API
diff --git a/Engine/Plugins/VirtualProduction/CaptureManager/CaptureManagerCore/Source/CaptureProtocolStack/Public/Discovery/DiscoveryMessenger.h b/Engine/Plugins/VirtualProduction/CaptureManager/CaptureManagerCore/Source/CaptureProtocolStack/Public/Discovery/DiscoveryMessenger.h
--- a/Engine/Plugins/VirtualProduction/CaptureManager/CaptureManagerCore/Source/CaptureProtocolStack/Public/Discovery/DiscoveryMessenger.h
+++ b/Engine/Plugins/VirtualProduction/CaptureManager/CaptureManagerCore/Source/CaptureProtocolStack/Public/Discovery/DiscoveryMessenger.h
@@ -11,29 +11,31 @@
 
 #include "Utility/Error.h"
 
+#define UE_API CAPTUREPROTOCOLSTACK_API
+
 namespace UE::CaptureManager
 {
 
-class CAPTUREPROTOCOLSTACK_API FDiscoveryMessenger final
+class FDiscoveryMessenger final
 {
 public:
 	using FOnResponseArrived = TDelegate<void(FString, FDiscoveryResponse), FDefaultTSDelegateUserPolicy>;
 	using FOnNotifyArrived = TDelegate<void(FString, FDiscoveryNotify), FDefaultTSDelegateUserPolicy>;
 
-	FDiscoveryMessenger();
-	~FDiscoveryMessenger();
+	UE_API FDiscoveryMessenger();
+	UE_API ~FDiscoveryMessenger();
 
-	TProtocolResult<void> Start();
-	TProtocolResult<void> Stop();
+	UE_API TProtocolResult<void> Start();
+	UE_API TProtocolResult<void> Stop();
 
-	TProtocolResult<void> SendRequest();
+	UE_API TProtocolResult<void> SendRequest();
 
-	void SetResponseHandler(FOnResponseArrived InOnResponse);
-	void SetNotifyHandler(FOnNotifyArrived InOnNotify);
+	UE_API void SetResponseHandler(FOnResponseArrived InOnResponse);
+	UE_API void SetNotifyHandler(FOnNotifyArrived InOnNotify);
 
 private:
 
-	void OnPacketArrived(FString InServerIp, FDiscoveryPacket InPacket);
+	UE_API void OnPacketArrived(FString InServerIp, FDiscoveryPacket InPacket);
 
 	FDiscoveryCommunication Communication;
 
@@ -41,4 +43,6 @@ private:
 	FOnNotifyArrived OnNotify;
 };
 
-}
\ No newline at end of file
+}
+
+#undef UE_API
diff --git a/Engine/Plugins/VirtualProduction/CaptureManager/CaptureManagerCore/Source/CaptureProtocolStack/Public/Discovery/Messages/DiscoveryNotify.h b/Engine/Plugins/VirtualProduction/CaptureManager/CaptureManagerCore/Source/CaptureProtocolStack/Public/Discovery/Messages/DiscoveryNotify.h
--- a/Engine/Plugins/VirtualProduction/CaptureManager/CaptureManagerCore/Source/CaptureProtocolStack/Public/Discovery/Messages/DiscoveryNotify.h
+++ b/Engine/Plugins/VirtualProduction/CaptureManager/CaptureManagerCore/Source/CaptureProtocolStack/Public/Discovery/Messages/DiscoveryNotify.h
@@ -9,14 +9,16 @@
 #include "Containers/Array.h"
 #include "Containers/StaticArray.h"
 
+#define UE_API CAPTUREPROTOCOLSTACK_API
+
 namespace UE::CaptureManager
 {
 
-class CAPTUREPROTOCOLSTACK_API FDiscoveryNotify
+class FDiscoveryNotify
 {
 public:
 
-	static const uint32 MinPayloadSize;
+	static UE_API const uint32 MinPayloadSize;
 
 	using FServerId = TStaticArray<uint8, 16>;
 
@@ -29,25 +31,25 @@ public:
 	};
 
 	UE_DEPRECATED(5.6, "This constructor is no longer supported")
-	FDiscoveryNotify(FServerId InServerId, uint16 InControlPort, EConnectionState InConnectionState, TArray<uint16> InSupportedVersions);
+	UE_API FDiscoveryNotify(FServerId InServerId, uint16 InControlPort, EConnectionState InConnectionState, TArray<uint16> InSupportedVersions);
 	
-	FDiscoveryNotify(FServerId InServerId, FString InServerName, uint16 InControlPort, EConnectionState InConnectionState);
+	UE_API FDiscoveryNotify(FServerId InServerId, FString InServerName, uint16 InControlPort, EConnectionState InConnectionState);
 
-	static TProtocolResult<FDiscoveryNotify> Deserialize(const FDiscoveryPacket& InPacket);
-	static TProtocolResult<FDiscoveryPacket> Serialize(const FDiscoveryNotify& InNotify);
+	static UE_API TProtocolResult<FDiscoveryNotify> Deserialize(const FDiscoveryPacket& InPacket);
+	static UE_API TProtocolResult<FDiscoveryPacket> Serialize(const FDiscoveryNotify& InNotify);
 
-	const FServerId& GetServerId() const;
-	const FString& GetServerName() const;
-	uint16 GetControlPort() const;
-	EConnectionState GetConnectionState() const;
+	UE_API const FServerId& GetServerId() const;
+	UE_API const FString& GetServerName() const;
+	UE_API uint16 GetControlPort() const;
+	UE_API EConnectionState GetConnectionState() const;
 	
 	UE_DEPRECATED(5.6, "GetSupportedVersions is no longer supported")
-	const TArray<uint16>& GetSupportedVersions() const;
+	UE_API const TArray<uint16>& GetSupportedVersions() const;
 
 private:
 
-	static EConnectionState ToConnectionState(uint8 InConnectionState);
-	static uint8 FromConnectionState(EConnectionState InConnectionState);
+	static UE_API EConnectionState ToConnectionState(uint8 InConnectionState);
+	static UE_API uint8 FromConnectionState(EConnectionState InConnectionState);
 
 	FServerId ServerId;
 	FString ServerName;
@@ -58,4 +60,6 @@ private:
 	TArray<uint16> SupportedVersions;
 };
 
-}
\ No newline at end of file
+}
+
+#undef UE_API
diff --git a/Engine/Plugins/VirtualProduction/CaptureManager/CaptureManagerCore/Source/CaptureProtocolStack/Public/Discovery/Messages/DiscoveryResponse.h b/Engine/Plugins/VirtualProduction/CaptureManager/CaptureManagerCore/Source/CaptureProtocolStack/Public/Discovery/Messages/DiscoveryResponse.h
--- a/Engine/Plugins/VirtualProduction/CaptureManager/CaptureManagerCore/Source/CaptureProtocolStack/Public/Discovery/Messages/DiscoveryResponse.h
+++ b/Engine/Plugins/VirtualProduction/CaptureManager/CaptureManagerCore/Source/CaptureProtocolStack/Public/Discovery/Messages/DiscoveryResponse.h
@@ -9,31 +9,33 @@
 #include "Containers/Array.h"
 #include "Containers/StaticArray.h"
 
+#define UE_API CAPTUREPROTOCOLSTACK_API
+
 namespace UE::CaptureManager
 {
 
-class CAPTUREPROTOCOLSTACK_API FDiscoveryResponse
+class FDiscoveryResponse
 {
 public:
 
-	static const uint32 MinPayloadSize;
+	static UE_API const uint32 MinPayloadSize;
 
 	using FServerId = TStaticArray<uint8, 16>;
 
 	UE_DEPRECATED(5.6, "This constructor is no longer supported")
-	FDiscoveryResponse(FServerId InServerId, uint16 InControlPort, TArray<uint16> InSupportedVersions);
+	UE_API FDiscoveryResponse(FServerId InServerId, uint16 InControlPort, TArray<uint16> InSupportedVersions);
 	
-	FDiscoveryResponse(FServerId InServerId, FString InServerName, uint16 InControlPort);
+	UE_API FDiscoveryResponse(FServerId InServerId, FString InServerName, uint16 InControlPort);
 
-	static TProtocolResult<FDiscoveryResponse> Deserialize(const FDiscoveryPacket& InPacket);
-	static TProtocolResult<FDiscoveryPacket> Serialize(const FDiscoveryResponse& InResponse);
+	static UE_API TProtocolResult<FDiscoveryResponse> Deserialize(const FDiscoveryPacket& InPacket);
+	static UE_API TProtocolResult<FDiscoveryPacket> Serialize(const FDiscoveryResponse& InResponse);
 
-	const FServerId& GetServerId() const;
-	const FString& GetServerName() const;
-	uint16 GetControlPort() const;
+	UE_API const FServerId& GetServerId() const;
+	UE_API const FString& GetServerName() const;
+	UE_API uint16 GetControlPort() const;
 	
 	UE_DEPRECATED(5.6, "GetSupportedVersions is no longer supported")
-	const TArray<uint16>& GetSupportedVersions() const;
+	UE_API const TArray<uint16>& GetSupportedVersions() const;
 
 private:
 
@@ -45,4 +47,6 @@ private:
 	TArray<uint16> SupportedVersions;
 };
 
-}
\ No newline at end of file
+}
+
+#undef UE_API
diff --git a/Engine/Plugins/VirtualProduction/CaptureManager/CaptureManagerCore/Source/CaptureProtocolStack/Public/ExportClient/ExportClient.h b/Engine/Plugins/VirtualProduction/CaptureManager/CaptureManagerCore/Source/CaptureProtocolStack/Public/ExportClient/ExportClient.h
--- a/Engine/Plugins/VirtualProduction/CaptureManager/CaptureManagerCore/Source/CaptureProtocolStack/Public/ExportClient/ExportClient.h
+++ b/Engine/Plugins/VirtualProduction/CaptureManager/CaptureManagerCore/Source/CaptureProtocolStack/Public/ExportClient/ExportClient.h
@@ -7,6 +7,8 @@
 
 #include "Containers/Map.h"
 
+#define UE_API CAPTUREPROTOCOLSTACK_API
+
 namespace UE::CaptureManager
 {
 
@@ -24,7 +26,7 @@ private:
 	std::atomic_bool bCanceled;
 };
 
-class CAPTUREPROTOCOLSTACK_API FExportClient final
+class FExportClient final
 	: public FExportTaskExecutor
 {
 public:
@@ -33,28 +35,28 @@ public:
 
 	using FTakeFileArray = TArray<FTakeFile>;
 
-	FExportClient(FString InServerIp, const uint16 InExportPort);
-	~FExportClient();
+	UE_API FExportClient(FString InServerIp, const uint16 InExportPort);
+	UE_API ~FExportClient();
 
-	FTaskId ExportTakeFiles(FString InTakeName, FTakeFileArray InTakeFileArray, TUniquePtr<FBaseStream> InStream);
-	FTaskId ExportFiles(TMap<FString, FTakeFileArray> InTakesFilesMap, TUniquePtr<FBaseStream> InStream);
+	UE_API FTaskId ExportTakeFiles(FString InTakeName, FTakeFileArray InTakeFileArray, TUniquePtr<FBaseStream> InStream);
+	UE_API FTaskId ExportFiles(TMap<FString, FTakeFileArray> InTakesFilesMap, TUniquePtr<FBaseStream> InStream);
 
-	void AbortExport(const FTaskId InTaskId);
-	void AbortAllExports();
+	UE_API void AbortExport(const FTaskId InTaskId);
+	UE_API void AbortAllExports();
 
 private:
 
-	FTaskId StartExport(FExportTakeTask::FExportContexts InContexts, TUniquePtr<FBaseStream> InStream);
+	UE_API FTaskId StartExport(FExportTakeTask::FExportContexts InContexts, TUniquePtr<FBaseStream> InStream);
 
-	virtual void OnTask(TUniquePtr<FExportTakeTask> InTask) override;
-	bool OnExportTask(TUniquePtr<FExportTakeTask> InTask);
+	UE_API virtual void OnTask(TUniquePtr<FExportTakeTask> InTask) override;
+	UE_API bool OnExportTask(TUniquePtr<FExportTakeTask> InTask);
 
-	TProtocolResult<TMap<uint32, TPair<FString, FTakeFile>>> SendRequests(FExportTakeTask::FExportContexts InContexts);
-	TProtocolResult<void> ReceiveResponses(TMap<uint32, TPair<FString, FTakeFile>> InResponseMap,
+	UE_API TProtocolResult<TMap<uint32, TPair<FString, FTakeFile>>> SendRequests(FExportTakeTask::FExportContexts InContexts);
+	UE_API TProtocolResult<void> ReceiveResponses(TMap<uint32, TPair<FString, FTakeFile>> InResponseMap,
 										   FBaseStream& InBaseStream);
 
-	TProtocolResult<void> Start();
-	TProtocolResult<void> Stop();
+	UE_API TProtocolResult<void> Start();
+	UE_API TProtocolResult<void> Stop();
 
 	FString ServerIp;
 	uint16 ServerPort = 0;
@@ -70,4 +72,6 @@ private:
 	TUniquePtr<FExportTaskStopToken> CurrentTaskStopToken;
 };
 
-}
\ No newline at end of file
+}
+
+#undef UE_API
diff --git a/Engine/Plugins/VirtualProduction/CaptureManager/CaptureManagerCore/Source/CaptureProtocolStack/Public/ExportClient/ExportWorker.h b/Engine/Plugins/VirtualProduction/CaptureManager/CaptureManagerCore/Source/CaptureProtocolStack/Public/ExportClient/ExportWorker.h
--- a/Engine/Plugins/VirtualProduction/CaptureManager/CaptureManagerCore/Source/CaptureProtocolStack/Public/ExportClient/ExportWorker.h
+++ b/Engine/Plugins/VirtualProduction/CaptureManager/CaptureManagerCore/Source/CaptureProtocolStack/Public/ExportClient/ExportWorker.h
@@ -96,7 +96,7 @@ private:
 	FExportTaskExecutor& Client;
 };
 
-class CAPTUREPROTOCOLSTACK_API FBaseStream
+class FBaseStream
 {
 public:
 
@@ -110,4 +110,4 @@ public:
 	virtual void Finalize(TProtocolResult<void> InResult) = 0;
 };
 
-}
\ No newline at end of file
+}
diff --git a/Engine/Plugins/VirtualProduction/CaptureManager/CaptureManagerCore/Source/CaptureUtils/Public/Async/CallbackSynchronizer.h b/Engine/Plugins/VirtualProduction/CaptureManager/CaptureManagerCore/Source/CaptureUtils/Public/Async/CallbackSynchronizer.h
--- a/Engine/Plugins/VirtualProduction/CaptureManager/CaptureManagerCore/Source/CaptureUtils/Public/Async/CallbackSynchronizer.h
+++ b/Engine/Plugins/VirtualProduction/CaptureManager/CaptureManagerCore/Source/CaptureUtils/Public/Async/CallbackSynchronizer.h
@@ -8,6 +8,8 @@
 
 #include <atomic>
 
+#define UE_API CAPTUREUTILS_API
+
 namespace UE::CaptureManager
 {
 
@@ -45,13 +47,13 @@ struct FFunctionTraits<Return(*)(Args ...)>
 	using FnType = Return(Args ...);
 };
 
-class CAPTUREUTILS_API FCallbackSynchronizer : public TSharedFromThis<FCallbackSynchronizer>
+class FCallbackSynchronizer : public TSharedFromThis<FCallbackSynchronizer>
 {
 public:
 	DECLARE_DELEGATE(FAfterAllDelegate);
 
-	FCallbackSynchronizer();
-	~FCallbackSynchronizer();
+	UE_API FCallbackSynchronizer();
+	UE_API ~FCallbackSynchronizer();
 
 	template<typename FCallback>
 	auto CreateCallback(FCallback&& InCallback)
@@ -60,9 +62,9 @@ public:
 		return FMakeCallbackInternal<FIdentityFnArgs>::Get(AsShared(), MoveTemp(InCallback));
 	}
 
-	void AfterAll(FAfterAllDelegate InAfterAllDelegate, bool bExecuteIfCounterZero = true);
+	UE_API void AfterAll(FAfterAllDelegate InAfterAllDelegate, bool bExecuteIfCounterZero = true);
 
-	static TSharedPtr<FCallbackSynchronizer> Create();
+	static UE_API TSharedPtr<FCallbackSynchronizer> Create();
 
 private:
 
@@ -87,11 +89,13 @@ private:
 		}
 	};
 
-	void Decrease();
+	UE_API void Decrease();
 
 	FAfterAllDelegate AfterAllDelegate;
 
 	std::atomic_int Counter;
 };
 
-}
\ No newline at end of file
+}
+
+#undef UE_API
diff --git a/Engine/Plugins/VirtualProduction/CaptureManager/CaptureManagerCore/Source/CaptureUtils/Public/Async/CaptureTimerManager.h b/Engine/Plugins/VirtualProduction/CaptureManager/CaptureManagerCore/Source/CaptureUtils/Public/Async/CaptureTimerManager.h
--- a/Engine/Plugins/VirtualProduction/CaptureManager/CaptureManagerCore/Source/CaptureUtils/Public/Async/CaptureTimerManager.h
+++ b/Engine/Plugins/VirtualProduction/CaptureManager/CaptureManagerCore/Source/CaptureUtils/Public/Async/CaptureTimerManager.h
@@ -8,22 +8,24 @@
 
 #include "TimerManager.h"
 
+#define UE_API CAPTUREUTILS_API
+
 namespace UE::CaptureManager
 {
 
-class CAPTUREUTILS_API FCaptureTimerManager
+class FCaptureTimerManager
 {
 public:
 
 	using FTimerHandle = FTSTicker::FDelegateHandle;
 
-	FCaptureTimerManager();
-	FCaptureTimerManager(const float Resolution);
+	UE_API FCaptureTimerManager();
+	UE_API FCaptureTimerManager(const float Resolution);
 
-	~FCaptureTimerManager();
+	UE_API ~FCaptureTimerManager();
 
-	FTimerHandle AddTimer(FTimerDelegate InDelegate, float InRate, bool InbLoop = false, float InFirstDelay = 0.f);
-	void RemoveTimer(FTimerHandle Handle);
+	UE_API FTimerHandle AddTimer(FTimerDelegate InDelegate, float InRate, bool InbLoop = false, float InFirstDelay = 0.f);
+	UE_API void RemoveTimer(FTimerHandle Handle);
 
 private:
 	static constexpr float IdealTimeResolution = 0.1f;
@@ -50,3 +52,5 @@ private:
 };
 
 }
+
+#undef UE_API
diff --git a/Engine/Plugins/VirtualProduction/CaptureManager/CaptureManagerCore/Source/CaptureUtils/Public/Async/Event.h b/Engine/Plugins/VirtualProduction/CaptureManager/CaptureManagerCore/Source/CaptureUtils/Public/Async/Event.h
--- a/Engine/Plugins/VirtualProduction/CaptureManager/CaptureManagerCore/Source/CaptureUtils/Public/Async/Event.h
+++ b/Engine/Plugins/VirtualProduction/CaptureManager/CaptureManagerCore/Source/CaptureUtils/Public/Async/Event.h
@@ -16,18 +16,20 @@
 
 #include <atomic>
 
+#define UE_API CAPTUREUTILS_API
+
 namespace UE::CaptureManager
 {
 
 // Base class for all event sublasses
-class CAPTUREUTILS_API FCaptureEvent
+class FCaptureEvent
 {
 public:
-	const FString& GetName() const;
-	virtual ~FCaptureEvent();
+	UE_API const FString& GetName() const;
+	UE_API virtual ~FCaptureEvent();
 
 protected:
-	FCaptureEvent(FString InName);
+	UE_API FCaptureEvent(FString InName);
 
 
 private:
@@ -57,4 +59,6 @@ public:
 	virtual void UnsubscribeAll() = 0;
 };
 
-}
\ No newline at end of file
+}
+
+#undef UE_API
diff --git a/Engine/Plugins/VirtualProduction/CaptureManager/CaptureManagerCore/Source/CaptureUtils/Public/Async/EventSourceUtils.h b/Engine/Plugins/VirtualProduction/CaptureManager/CaptureManagerCore/Source/CaptureUtils/Public/Async/EventSourceUtils.h
--- a/Engine/Plugins/VirtualProduction/CaptureManager/CaptureManagerCore/Source/CaptureUtils/Public/Async/EventSourceUtils.h
+++ b/Engine/Plugins/VirtualProduction/CaptureManager/CaptureManagerCore/Source/CaptureUtils/Public/Async/EventSourceUtils.h
@@ -4,24 +4,26 @@
 
 #include "Async/Event.h"
 
+#define UE_API CAPTUREUTILS_API
+
 namespace UE::CaptureManager
 {
 namespace Private
 {
-	class CAPTUREUTILS_API FCaptureEventSourceBase : public ICaptureEventSource
+	class FCaptureEventSourceBase : public ICaptureEventSource
 	{
 	public:
-		virtual ~FCaptureEventSourceBase();
+		UE_API virtual ~FCaptureEventSourceBase();
 
 	public:
-		virtual TArray<FString> GetAvailableEvents() const override;
+		UE_API virtual TArray<FString> GetAvailableEvents() const override;
 		// NOTE: SubscribeToEvent() must not be called from an event handler that is being executed on the same thread
-		virtual void SubscribeToEvent(const FString& InEventName, FCaptureEventHandler InHandler) override;
-		virtual void UnsubscribeAll() override;
+		UE_API virtual void SubscribeToEvent(const FString& InEventName, FCaptureEventHandler InHandler) override;
+		UE_API virtual void UnsubscribeAll() override;
 
 	protected:
-		void RegisterEvent(const FString& InEventName);
-		void PublishEventInternal(TSharedPtr<const FCaptureEvent> InEvent) const;
+		UE_API void RegisterEvent(const FString& InEventName);
+		UE_API void PublishEventInternal(TSharedPtr<const FCaptureEvent> InEvent) const;
 
 	private:
 		using FHandlers = TTSMulticastDelegate<void(TSharedPtr<const FCaptureEvent>)>;
@@ -33,7 +35,7 @@ namespace Private
 }
 
 // Class to be publicly inherited to get basic event source functionality. All functions are thread-safe.
-class CAPTUREUTILS_API FCaptureEventSource : public Private::FCaptureEventSourceBase
+class FCaptureEventSource : public Private::FCaptureEventSourceBase
 {
 protected:
 	template<typename EventType, typename ... Args>
@@ -52,7 +54,7 @@ private:
 };
 
 // Class to be publicly inherited to get event source functionality where rate of events published is limited. All functions are thread-safe.
-class CAPTUREUTILS_API FCaptureEventSourceWithLimiter : public Private::FCaptureEventSourceBase
+class FCaptureEventSourceWithLimiter : public Private::FCaptureEventSourceBase
 {
 public:
 	FCaptureEventSourceWithLimiter(int32 InThresholdMillis) : ThresholdMillis(InThresholdMillis), LastPublish(0)
@@ -110,4 +112,6 @@ private:
 	std::atomic<int64> LastPublish;
 };
 
-}
\ No newline at end of file
+}
+
+#undef UE_API
diff --git a/Engine/Plugins/VirtualProduction/CaptureManager/CaptureManagerCore/Source/CaptureUtils/Public/Async/StopToken.h b/Engine/Plugins/VirtualProduction/CaptureManager/CaptureManagerCore/Source/CaptureUtils/Public/Async/StopToken.h
--- a/Engine/Plugins/VirtualProduction/CaptureManager/CaptureManagerCore/Source/CaptureUtils/Public/Async/StopToken.h
+++ b/Engine/Plugins/VirtualProduction/CaptureManager/CaptureManagerCore/Source/CaptureUtils/Public/Async/StopToken.h
@@ -4,12 +4,14 @@
 
 #include "Templates/SharedPointer.h"
 
+#define UE_API CAPTUREUTILS_API
+
 namespace UE::CaptureManager
 {
 
 struct FSharedState;
 
-class CAPTUREUTILS_API FStopToken
+class FStopToken
 {
 public:
 
@@ -19,31 +21,31 @@ public:
 	FStopToken& operator=(const FStopToken& InOther) = default;
 	FStopToken& operator=(FStopToken&& InOther) = default;
 
-	bool IsStopRequested() const;
+	UE_API bool IsStopRequested() const;
 
 private:
 
 	friend class FStopRequester;
 
-	FStopToken(TWeakPtr<const FSharedState> InSharedState);
+	UE_API FStopToken(TWeakPtr<const FSharedState> InSharedState);
 
 	TWeakPtr<const FSharedState> SharedStateWeak;
 };
 
-class CAPTUREUTILS_API FStopRequester
+class FStopRequester
 {
 public:
-	FStopRequester();
+	UE_API FStopRequester();
 
 	FStopRequester(const FStopRequester& InOther) = default;
 	FStopRequester(FStopRequester&& InOther) = default;
 	FStopRequester& operator=(const FStopRequester& InOther) = default;
 	FStopRequester& operator=(FStopRequester&& InOther) = default;
 
-	void RequestStop();
-	bool IsStopRequested() const;
+	UE_API void RequestStop();
+	UE_API bool IsStopRequested() const;
 
-	FStopToken CreateToken() const;
+	UE_API FStopToken CreateToken() const;
 
 private:
 
@@ -51,3 +53,5 @@ private:
 };
 
 }
+
+#undef UE_API
diff --git a/Engine/Plugins/VirtualProduction/CaptureManager/CaptureManagerCore/Source/CaptureUtils/Public/Async/Task.h b/Engine/Plugins/VirtualProduction/CaptureManager/CaptureManagerCore/Source/CaptureUtils/Public/Async/Task.h
--- a/Engine/Plugins/VirtualProduction/CaptureManager/CaptureManagerCore/Source/CaptureUtils/Public/Async/Task.h
+++ b/Engine/Plugins/VirtualProduction/CaptureManager/CaptureManagerCore/Source/CaptureUtils/Public/Async/Task.h
@@ -5,27 +5,29 @@
 #include "Templates/Function.h"
 #include "Async/StopToken.h"
 
+#define UE_API CAPTUREUTILS_API
+
 template<typename TTask>
 class FAsyncTask;
 
 namespace UE::CaptureManager
 {
 
-class CAPTUREUTILS_API FCancelableAsyncTask
+class FCancelableAsyncTask
 {
 public:
 	using FTaskFunction = TFunction<void(const FStopToken&)>;
 
-	explicit FCancelableAsyncTask(FTaskFunction InTaskFunction);
+	UE_API explicit FCancelableAsyncTask(FTaskFunction InTaskFunction);
 
-	~FCancelableAsyncTask();
+	UE_API ~FCancelableAsyncTask();
 
-	bool IsDone();
+	UE_API bool IsDone();
 
-	void StartSync();
-	void StartAsync();
+	UE_API void StartSync();
+	UE_API void StartAsync();
 
-	void Cancel();
+	UE_API void Cancel();
 
 private:
 	class FAsyncTaskInternal;
@@ -35,3 +37,5 @@ private:
 };
 
 }
+
+#undef UE_API
diff --git a/Engine/Plugins/VirtualProduction/CaptureManager/CaptureManagerCore/Source/CaptureUtils/Public/Async/TaskWaiter.h b/Engine/Plugins/VirtualProduction/CaptureManager/CaptureManagerCore/Source/CaptureUtils/Public/Async/TaskWaiter.h
--- a/Engine/Plugins/VirtualProduction/CaptureManager/CaptureManagerCore/Source/CaptureUtils/Public/Async/TaskWaiter.h
+++ b/Engine/Plugins/VirtualProduction/CaptureManager/CaptureManagerCore/Source/CaptureUtils/Public/Async/TaskWaiter.h
@@ -3,19 +3,21 @@
 #include <atomic>
 #include "HAL/Platform.h"
 
+#define UE_API CAPTUREUTILS_API
+
 namespace UE::CaptureManager
 {
 
-class CAPTUREUTILS_API FTaskWaiter
+class FTaskWaiter
 {
 public:
 
-	FTaskWaiter();
-	~FTaskWaiter();
+	UE_API FTaskWaiter();
+	UE_API ~FTaskWaiter();
 
-	bool CreateTask();
-	void FinishTask();
-	void WaitForAll();
+	UE_API bool CreateTask();
+	UE_API void FinishTask();
+	UE_API void WaitForAll();
 
 private:
 
@@ -23,4 +25,6 @@ private:
 	const uint32 CanCreateTaskFlag = 0x80000000;
 };
 
-}
\ No newline at end of file
+}
+
+#undef UE_API
diff --git a/Engine/Plugins/VirtualProduction/CaptureManager/CaptureManagerCore/Source/CaptureUtils/Public/CaptureUtilsModule.h b/Engine/Plugins/VirtualProduction/CaptureManager/CaptureManagerCore/Source/CaptureUtils/Public/CaptureUtilsModule.h
--- a/Engine/Plugins/VirtualProduction/CaptureManager/CaptureManagerCore/Source/CaptureUtils/Public/CaptureUtilsModule.h
+++ b/Engine/Plugins/VirtualProduction/CaptureManager/CaptureManagerCore/Source/CaptureUtils/Public/CaptureUtilsModule.h
@@ -6,16 +6,20 @@
 
 #include "Async/CaptureTimerManager.h"
 
-class CAPTUREUTILS_API FCaptureUtilsModule : public IModuleInterface
+#define UE_API CAPTUREUTILS_API
+
+class FCaptureUtilsModule : public IModuleInterface
 {
 public:
 
-	virtual void StartupModule() override;
-	virtual void ShutdownModule() override;
+	UE_API virtual void StartupModule() override;
+	UE_API virtual void ShutdownModule() override;
 
-	TSharedRef<UE::CaptureManager::FCaptureTimerManager> GetTimerManager();
+	UE_API TSharedRef<UE::CaptureManager::FCaptureTimerManager> GetTimerManager();
 
 private:
 
 	TSharedPtr<UE::CaptureManager::FCaptureTimerManager> TimerManager;
-};
\ No newline at end of file
+};
+
+#undef UE_API
diff --git a/Engine/Plugins/VirtualProduction/CaptureManager/CaptureManagerCore/Source/CaptureUtils/Public/Network/Error.h b/Engine/Plugins/VirtualProduction/CaptureManager/CaptureManagerCore/Source/CaptureUtils/Public/Network/Error.h
--- a/Engine/Plugins/VirtualProduction/CaptureManager/CaptureManagerCore/Source/CaptureUtils/Public/Network/Error.h
+++ b/Engine/Plugins/VirtualProduction/CaptureManager/CaptureManagerCore/Source/CaptureUtils/Public/Network/Error.h
@@ -6,15 +6,17 @@
 
 #include "Templates/ValueOrError.h"
 
+#define UE_API CAPTUREUTILS_API
+
 namespace UE::CaptureManager
 {
 
-class CAPTUREUTILS_API FCaptureProtocolError
+class FCaptureProtocolError
 {
 public:
 
-	FCaptureProtocolError();
-	FCaptureProtocolError(FString InMessage, int32 InCode = 0);
+	UE_API FCaptureProtocolError();
+	UE_API FCaptureProtocolError(FString InMessage, int32 InCode = 0);
 
 	FCaptureProtocolError(const FCaptureProtocolError& InOther) = default;
 	FCaptureProtocolError(FCaptureProtocolError&& InOther) = default;
@@ -22,8 +24,8 @@ public:
 	FCaptureProtocolError& operator=(const FCaptureProtocolError& InOther) = default;
 	FCaptureProtocolError& operator=(FCaptureProtocolError&& InOther) = default;
 
-	const FString& GetMessage() const;
-	int32 GetCode() const;
+	UE_API const FString& GetMessage() const;
+	UE_API int32 GetCode() const;
 
 private:
 
@@ -106,4 +108,6 @@ public:
 
 const TProtocolResult<void> ResultOk = TInPlaceType<void>{};
 
-}
\ No newline at end of file
+}
+
+#undef UE_API
diff --git a/Engine/Plugins/VirtualProduction/CaptureManager/CaptureManagerCore/Source/CaptureUtils/Public/Network/TcpClient.h b/Engine/Plugins/VirtualProduction/CaptureManager/CaptureManagerCore/Source/CaptureUtils/Public/Network/TcpClient.h
--- a/Engine/Plugins/VirtualProduction/CaptureManager/CaptureManagerCore/Source/CaptureUtils/Public/Network/TcpClient.h
+++ b/Engine/Plugins/VirtualProduction/CaptureManager/CaptureManagerCore/Source/CaptureUtils/Public/Network/TcpClient.h
@@ -20,12 +20,14 @@
 
 #include "Misc/DateTime.h"
 
+#define UE_API CAPTUREUTILS_API
+
 namespace UE::CaptureManager
 {
 
 using FSocketPtr = TUniquePtr<FSocket, FSocketDeleter>;
 
-class CAPTUREUTILS_API FTcpClient final
+class FTcpClient final
 {
 public:
 
@@ -35,18 +37,18 @@ public:
 	static constexpr int32 NoPendingDataError = -2;
 	static constexpr int32 TimeoutError = -1;
 
-	FTcpClient();
-	~FTcpClient();
+	UE_API FTcpClient();
+	UE_API ~FTcpClient();
 
-	TProtocolResult<void> Init();
+	UE_API TProtocolResult<void> Init();
 	// Blocking function until connection is established
-	TProtocolResult<void> Start(const FString& InServerAddress);
-	TProtocolResult<void> Stop();
+	UE_API TProtocolResult<void> Start(const FString& InServerAddress);
+	UE_API TProtocolResult<void> Stop();
 
-	bool IsRunning() const;
+	UE_API bool IsRunning() const;
 
-	TProtocolResult<void> SendMessage(const TArray<uint8>& InPayload);
-	TProtocolResult<TArray<uint8>> ReceiveMessage(const uint64 InSize, const uint32 InWaitTimeoutMs = ITcpSocketReader::DefaultWaitTimeoutMs);
+	UE_API TProtocolResult<void> SendMessage(const TArray<uint8>& InPayload);
+	UE_API TProtocolResult<TArray<uint8>> ReceiveMessage(const uint64 InSize, const uint32 InWaitTimeoutMs = ITcpSocketReader::DefaultWaitTimeoutMs);
 
 private:
 
@@ -55,27 +57,29 @@ private:
 	bool bRunning;
 };
 
-class CAPTUREUTILS_API FTcpClientReader final : public ITcpSocketReader
+class FTcpClientReader final : public ITcpSocketReader
 {
 public:
-	FTcpClientReader(FTcpClient& InClient);
+	UE_API FTcpClientReader(FTcpClient& InClient);
 
-	virtual TProtocolResult<TArray<uint8>> ReceiveMessage(const uint64 InSize, const uint32 InWaitTimeoutMs = DefaultWaitTimeoutMs) override;
+	UE_API virtual TProtocolResult<TArray<uint8>> ReceiveMessage(const uint64 InSize, const uint32 InWaitTimeoutMs = DefaultWaitTimeoutMs) override;
 
 private:
 	FTcpClient& Client;
 };
 
-class CAPTUREUTILS_API FTcpClientWriter final : public ITcpSocketWriter
+class FTcpClientWriter final : public ITcpSocketWriter
 {
 public:
-	FTcpClientWriter(FTcpClient& InClient);
+	UE_API FTcpClientWriter(FTcpClient& InClient);
 
-	virtual TProtocolResult<void> SendMessage(const TArray<uint8>& InPayload) override;
+	UE_API virtual TProtocolResult<void> SendMessage(const TArray<uint8>& InPayload) override;
 
 private:
 
 	FTcpClient& Client;
 };
 
-}
\ No newline at end of file
+}
+
+#undef UE_API
diff --git a/Engine/Plugins/VirtualProduction/CaptureManager/CaptureManagerCore/Source/CaptureUtils/Public/Network/TcpServer.h b/Engine/Plugins/VirtualProduction/CaptureManager/CaptureManagerCore/Source/CaptureUtils/Public/Network/TcpServer.h
--- a/Engine/Plugins/VirtualProduction/CaptureManager/CaptureManagerCore/Source/CaptureUtils/Public/Network/TcpServer.h
+++ b/Engine/Plugins/VirtualProduction/CaptureManager/CaptureManagerCore/Source/CaptureUtils/Public/Network/TcpServer.h
@@ -21,12 +21,14 @@
 
 #include "Delegates/Delegate.h"
 
+#define UE_API CAPTUREUTILS_API
+
 namespace UE::CaptureManager
 {
 
 using FSocketPtr = TUniquePtr<FSocket, FSocketDeleter>;
 
-class CAPTUREUTILS_API FTcpClientHandler final
+class FTcpClientHandler final
 {
 public:
 
@@ -34,17 +36,17 @@ public:
 	static constexpr int32 DisconnectedError = -10;
 	static constexpr int32 TimeoutError = -1;
 
-	FTcpClientHandler(FSocketPtr InSocket, FString InEndpoint);
-	~FTcpClientHandler();
+	UE_API FTcpClientHandler(FSocketPtr InSocket, FString InEndpoint);
+	UE_API ~FTcpClientHandler();
 
-	TProtocolResult<void> SendMessage(const TArray<uint8>& InData);
-	TProtocolResult<TArray<uint8>> ReceiveMessage(const uint64 InSize, const uint32 InWaitTimeoutMs = ITcpSocketReader::DefaultWaitTimeoutMs);
+	UE_API TProtocolResult<void> SendMessage(const TArray<uint8>& InData);
+	UE_API TProtocolResult<TArray<uint8>> ReceiveMessage(const uint64 InSize, const uint32 InWaitTimeoutMs = ITcpSocketReader::DefaultWaitTimeoutMs);
 
-	TProtocolResult<uint32> HasPendingData() const;
+	UE_API TProtocolResult<uint32> HasPendingData() const;
 
-	const FString& GetEndpoint() const;
+	UE_API const FString& GetEndpoint() const;
 
-	bool operator==(const FTcpClientHandler& InOther);
+	UE_API bool operator==(const FTcpClientHandler& InOther);
 
 private:
 
@@ -52,7 +54,7 @@ private:
 	FString Endpoint;
 };
 
-class CAPTUREUTILS_API FTcpServer
+class FTcpServer
 {
 public:
 
@@ -61,19 +63,19 @@ public:
 
 	DECLARE_DELEGATE_TwoParams(FConnectionHandler, TWeakPtr<FTcpClientHandler> InClient, bool bConnected)
 
-	FTcpServer(const uint32 InMaxNumberOfClients);
+	UE_API FTcpServer(const uint32 InMaxNumberOfClients);
 
-	TProtocolResult<uint16> Start(const uint16 InListenPort = AnyPort);
-	TProtocolResult<void> Stop();
+	UE_API TProtocolResult<uint16> Start(const uint16 InListenPort = AnyPort);
+	UE_API TProtocolResult<void> Stop();
 
-	bool IsRunning() const;
+	UE_API bool IsRunning() const;
 
-	TProtocolResult<void> SendMessage(const TArray<uint8>& InMessage, const FString& InEndpoint);
+	UE_API TProtocolResult<void> SendMessage(const TArray<uint8>& InMessage, const FString& InEndpoint);
 
-	void DisconnectClient(const FString& InEndpoint);
-	void SetConnectionHandler(FConnectionHandler InOnConnectionHandler);
+	UE_API void DisconnectClient(const FString& InEndpoint);
+	UE_API void SetConnectionHandler(FConnectionHandler InOnConnectionHandler);
 
-	int32 GetPort() const;
+	UE_API int32 GetPort() const;
 
 private:
 
@@ -89,27 +91,29 @@ private:
 	std::atomic_bool bRunning;
 };
 
-class CAPTUREUTILS_API FTcpConnectionReader final : public ITcpSocketReader
+class FTcpConnectionReader final : public ITcpSocketReader
 {
 public:
-	FTcpConnectionReader(FTcpClientHandler& InClient);
+	UE_API FTcpConnectionReader(FTcpClientHandler& InClient);
 
-	virtual TProtocolResult<TArray<uint8>> ReceiveMessage(const uint64 InSize, const uint32 InWaitTimeoutMs = DefaultWaitTimeoutMs) override;
+	UE_API virtual TProtocolResult<TArray<uint8>> ReceiveMessage(const uint64 InSize, const uint32 InWaitTimeoutMs = DefaultWaitTimeoutMs) override;
 
 private:
 	FTcpClientHandler& Client;
 };
 
-class CAPTUREUTILS_API FTcpConnectionWriter final : public ITcpSocketWriter
+class FTcpConnectionWriter final : public ITcpSocketWriter
 {
 public:
-	FTcpConnectionWriter(FTcpClientHandler& InClient);
+	UE_API FTcpConnectionWriter(FTcpClientHandler& InClient);
 
-	virtual TProtocolResult<void> SendMessage(const TArray<uint8>& InPayload) override;
+	UE_API virtual TProtocolResult<void> SendMessage(const TArray<uint8>& InPayload) override;
 
 private:
 
 	FTcpClientHandler& Client;
 };
 
-}
\ No newline at end of file
+}
+
+#undef UE_API
diff --git a/Engine/Plugins/VirtualProduction/CaptureManager/CaptureManagerCore/Source/CaptureUtils/Public/Network/UdpClient.h b/Engine/Plugins/VirtualProduction/CaptureManager/CaptureManagerCore/Source/CaptureUtils/Public/Network/UdpClient.h
--- a/Engine/Plugins/VirtualProduction/CaptureManager/CaptureManagerCore/Source/CaptureUtils/Public/Network/UdpClient.h
+++ b/Engine/Plugins/VirtualProduction/CaptureManager/CaptureManagerCore/Source/CaptureUtils/Public/Network/UdpClient.h
@@ -18,6 +18,8 @@
 
 #include "Templates/UniquePtr.h"
 
+#define UE_API CAPTUREUTILS_API
+
 namespace UE::CaptureManager
 {
 
@@ -29,18 +31,18 @@ struct FUdpClientConfigure
 
 #define CHECK_BOOL(Function) if (bool Result = Function; !Result) { return FCaptureProtocolError("Failed to execute function"); }
 
-class CAPTUREUTILS_API FUdpClient final
+class FUdpClient final
 {
 public:
 
-	FUdpClient();
-	~FUdpClient();
+	UE_API FUdpClient();
+	UE_API ~FUdpClient();
 
-	TProtocolResult<void> Init(FUdpClientConfigure InConfig, FOnSocketDataReceived InReceiveHandler);
-	TProtocolResult<void> Start();
-	TProtocolResult<void> Stop();
+	UE_API TProtocolResult<void> Init(FUdpClientConfigure InConfig, FOnSocketDataReceived InReceiveHandler);
+	UE_API TProtocolResult<void> Start();
+	UE_API TProtocolResult<void> Stop();
 
-	TProtocolResult<int32> SendMessage(const TArray<uint8>& InPayload, const FString& InEndpoint);
+	UE_API TProtocolResult<int32> SendMessage(const TArray<uint8>& InPayload, const FString& InEndpoint);
 
 private:
 
@@ -54,4 +56,6 @@ private:
 	bool bRunning;
 };
 
-}
\ No newline at end of file
+}
+
+#undef UE_API
diff --git a/Engine/Plugins/VirtualProduction/CaptureManager/CaptureManagerCore/Source/DataIngestCore/Public/IngestCaptureData.h b/Engine/Plugins/VirtualProduction/CaptureManager/CaptureManagerCore/Source/DataIngestCore/Public/IngestCaptureData.h
--- a/Engine/Plugins/VirtualProduction/CaptureManager/CaptureManagerCore/Source/DataIngestCore/Public/IngestCaptureData.h
+++ b/Engine/Plugins/VirtualProduction/CaptureManager/CaptureManagerCore/Source/DataIngestCore/Public/IngestCaptureData.h
@@ -5,11 +5,13 @@
 #include "Templates/ValueOrError.h"
 #include "Containers/UnrealString.h"
 
+#define UE_API DATAINGESTCORE_API
+
 /** Capture archive data object for import into UE. */
-struct DATAINGESTCORE_API FIngestCaptureData
+struct FIngestCaptureData
 {
 	/** Expected extension for capture archive files. */
-	static const FString Extension;
+	static UE_API const FString Extension;
 
 	/** Video/image sequence information. */
 	struct FVideo
@@ -65,4 +67,6 @@ DATAINGESTCORE_API FParseResult ParseFile(const FString& InFilePath);
 /** Serialise capture archive data object to file. */
 DATAINGESTCORE_API TOptional<FText> Serialize(const FString& InFilePath, const FString& InFileName, const FIngestCaptureData& InIngestCaptureData);
 
-}
\ No newline at end of file
+}
+
+#undef UE_API
