Handling GetTakeList response should check for take name uniqueness

diff --git a/Engine/Plugins/VirtualProduction/CaptureManager/CaptureManagerCore/Source/CaptureProtocolStack/Private/Control/ControlMessenger.cpp b/Engine/Plugins/VirtualProduction/CaptureManager/CaptureManagerCore/Source/CaptureProtocolStack/Private/Control/ControlMessenger.cpp
--- a/Engine/Plugins/VirtualProduction/CaptureManager/CaptureManagerCore/Source/CaptureProtocolStack/Private/Control/ControlMessenger.cpp
+++ b/Engine/Plugins/VirtualProduction/CaptureManager/CaptureManagerCore/Source/CaptureProtocolStack/Private/Control/ControlMessenger.cpp
@@ -118,6 +118,31 @@ TProtocolResult<FGetServerInformationResponse> FControlMessenger::GetServerInfor
 	return SendRequest(FGetServerInformationRequest());
 }
 
+TProtocolResult<FGetTakeListResponse> FControlMessenger::GetTakeList()
+{
+	TProtocolResult<FGetTakeListResponse> Response =
+		SendRequest(FGetTakeListRequest());
+
+	if (Response.HasError())
+	{
+		return Response;
+	}
+
+	TArray<FString> TakeNames = Response.GetValue().GetNames();
+	TakeNames.Sort();
+	
+	for (int32 Index = 0; Index < TakeNames.Num() - 1; ++Index)
+	{
+		if (TakeNames[Index] == TakeNames[Index + 1])
+		{
+			FString Message = FString::Format(TEXT("Response for Get Take List Request contains duplicate take name: {0}. Take names are expected to be unique"), { TakeNames[Index] });
+			return FCaptureProtocolError(MoveTemp(Message));
+		}
+	}
+
+	return Response;
+}
+
 void FControlMessenger::SendPacket(FControlPacket InPacket)
 {
 	Communication.SendMessage(MoveTemp(InPacket));
diff --git a/Engine/Plugins/VirtualProduction/CaptureManager/CaptureManagerCore/Source/CaptureProtocolStack/Public/Control/ControlMessenger.h b/Engine/Plugins/VirtualProduction/CaptureManager/CaptureManagerCore/Source/CaptureProtocolStack/Public/Control/ControlMessenger.h
--- a/Engine/Plugins/VirtualProduction/CaptureManager/CaptureManagerCore/Source/CaptureProtocolStack/Public/Control/ControlMessenger.h
+++ b/Engine/Plugins/VirtualProduction/CaptureManager/CaptureManagerCore/Source/CaptureProtocolStack/Public/Control/ControlMessenger.h
@@ -62,6 +62,7 @@ public:
 
 	UE_API TProtocolResult<void> StartSession();
 	UE_API TProtocolResult<FGetServerInformationResponse> GetServerInformation();
+	UE_API TProtocolResult<FGetTakeListResponse> GetTakeList();
 
 	template<class RequestType>
 	TProtocolResult<typename RequestType::ResponseType> SendRequest(RequestType InRequest)
diff --git a/Engine/Plugins/VirtualProduction/CaptureManager/CaptureManagerDevices/Source/CPSLiveLinkDevice/Private/Protocol/CPSDevice.cpp b/Engine/Plugins/VirtualProduction/CaptureManager/CaptureManagerDevices/Source/CPSLiveLinkDevice/Private/Protocol/CPSDevice.cpp
--- a/Engine/Plugins/VirtualProduction/CaptureManager/CaptureManagerDevices/Source/CPSLiveLinkDevice/Private/Protocol/CPSDevice.cpp
+++ b/Engine/Plugins/VirtualProduction/CaptureManager/CaptureManagerDevices/Source/CPSLiveLinkDevice/Private/Protocol/CPSDevice.cpp
@@ -165,7 +165,7 @@ TProtocolResult<TArray<FGetTakeMetadataResponse::FTakeObject>> FCPSDevice::Fetch
 		return FCaptureProtocolError(MoveTemp(Message));
 	}
 
-	TProtocolResult<FGetTakeListResponse> TakeListRequestResult = ControlMessenger.SendRequest(FGetTakeListRequest());
+	TProtocolResult<FGetTakeListResponse> TakeListRequestResult = ControlMessenger.GetTakeList();
 	CHECK_LOG_RETURN_VAL(TakeListRequestResult.HasError(), TEXT("Failed to fetch the take list from the remote host."));
 
 	FGetTakeListResponse TakeListResponse = TakeListRequestResult.StealValue();
