Unshelved from pending changelist '46101232':

[CaptureManager] Fix crash when accessing an unset TOptional

#was-backout 46101187

diff --git a/Engine/Plugins/VirtualProduction/CaptureManager/CaptureManagerCore/Source/LiveLinkHubCaptureMessaging/Private/UploadDataMessage.cpp b/Engine/Plugins/VirtualProduction/CaptureManager/CaptureManagerCore/Source/LiveLinkHubCaptureMessaging/Private/UploadDataMessage.cpp
--- a/Engine/Plugins/VirtualProduction/CaptureManager/CaptureManagerCore/Source/LiveLinkHubCaptureMessaging/Private/UploadDataMessage.cpp
+++ b/Engine/Plugins/VirtualProduction/CaptureManager/CaptureManagerCore/Source/LiveLinkHubCaptureMessaging/Private/UploadDataMessage.cpp
@@ -137,14 +137,45 @@ FUploadResult<FUploadDataHeader> FUploadDataMessage::DeserializeHeader(UE::Captu
 {
 	FUploadDataHeader Header;
 
-	CHECK_AND_RET_ERROR(DeserializeStartHeader(InReader));
-	CHECK_AND_RET_ERROR(DeserializeGuid(InReader, Header.ClientId));
-	CHECK_AND_RET_ERROR(DeserializeGuid(InReader, Header.CaptureSourceId));
-	CHECK_AND_RET_ERROR(DeserializeString(InReader, Header.CaptureSourceName));
-	CHECK_AND_RET_ERROR(DeserializeGuid(InReader, Header.TakeUploadId));
-	CHECK_AND_RET_ERROR(DeserializeString(InReader, Header.Slate));
-	CHECK_AND_RET_ERROR(DeserializeTakeNumber(InReader, Header.TakeNumber));
-	CHECK_AND_RET_ERROR(DeserializeTotalLength(InReader, Header.TotalLength));
+	{
+		FUploadVoidResult Result = DeserializeStartHeader(InReader);
+		CHECK_AND_RET_ERROR(Result);
+	}
+
+	{
+		FUploadVoidResult Result = DeserializeGuid(InReader, Header.ClientId);
+		CHECK_AND_RET_ERROR(Result);
+	}
+
+	{
+		FUploadVoidResult Result = DeserializeGuid(InReader, Header.CaptureSourceId);
+		CHECK_AND_RET_ERROR(Result);
+	}
+
+	{
+		FUploadVoidResult Result = DeserializeString(InReader, Header.CaptureSourceName);
+		CHECK_AND_RET_ERROR(Result);
+	}
+
+	{
+		FUploadVoidResult Result = DeserializeGuid(InReader, Header.TakeUploadId);
+		CHECK_AND_RET_ERROR(Result);
+	}
+
+	{
+		FUploadVoidResult Result = DeserializeString(InReader, Header.Slate);
+		CHECK_AND_RET_ERROR(Result);
+	}
+
+	{
+		FUploadVoidResult Result = DeserializeTakeNumber(InReader, Header.TakeNumber);
+		CHECK_AND_RET_ERROR(Result);
+	}
+
+	{
+		FUploadVoidResult Result = DeserializeTotalLength(InReader, Header.TotalLength);
+		CHECK_AND_RET_ERROR(Result);
+	}
 
 	return MakeValue(MoveTemp(Header));
 }
@@ -153,8 +184,15 @@ FUploadResult<FUploadFileDataHeader> FUploadDataMessage::DeserializeFileHeader(U
 {
 	FUploadFileDataHeader FileHeader;
 
-	CHECK_AND_RET_ERROR(DeserializeFileName(InReader, FileHeader.FileName));
-	CHECK_AND_RET_ERROR(DeserializeLength(InReader, FileHeader.Length));
+	{
+		FUploadVoidResult Result = DeserializeFileName(InReader, FileHeader.FileName);
+		CHECK_AND_RET_ERROR(Result);
+	}
+
+	{
+		FUploadVoidResult Result = DeserializeLength(InReader, FileHeader.Length);
+		CHECK_AND_RET_ERROR(Result);
+	}
 
 	return MakeValue(MoveTemp(FileHeader));
 }
