Fix crash on Mac when ingesting using ffmpeg

Adds a dedicated process runner static class which encapsulates all capture manager external process execution. Includes stop token and timeout support.

Regarding the Mac crash specifically, ffmpeg child processes were immediately receiving SIGTTOU (signal 22) and stopping (suspended) because they inherited the parent's terminal stdin while being placed in a background process group (POSIX_SPAWN_SETPGROUP). When ffmpeg tried to interact with the terminal, macOS job control suspended it, causing the process reaper to crash the engine. It seems the UE platform reaper functionality does not handle suspended processes well.

By attaching a dedicated stdin pipe pair to the child process it prevents ffmpeg from trying access the terminal when LLH is executed from the command line.

#rnx
#Synced-CL 48783767

diff --git a/Engine/Plugins/VirtualProduction/CaptureManager/CaptureManagerApp/Source/CaptureDataConverter/Private/Nodes/ThirdPartyEncoder/CaptureConvertAudioDataThirdParty.cpp b/Engine/Plugins/VirtualProduction/CaptureManager/CaptureManagerApp/Source/CaptureDataConverter/Private/Nodes/ThirdPartyEncoder/CaptureConvertAudioDataThirdParty.cpp
--- a/Engine/Plugins/VirtualProduction/CaptureManager/CaptureManagerApp/Source/CaptureDataConverter/Private/Nodes/ThirdPartyEncoder/CaptureConvertAudioDataThirdParty.cpp
+++ b/Engine/Plugins/VirtualProduction/CaptureManager/CaptureManagerApp/Source/CaptureDataConverter/Private/Nodes/ThirdPartyEncoder/CaptureConvertAudioDataThirdParty.cpp
@@ -20,6 +20,8 @@
 #include "Settings/CaptureManagerSettings.h"
 #include "Settings/CaptureManagerTemplateTokens.h"
 
+#include "ProcessRunner/ProcessRunner.h"
+
 #define LOCTEXT_NAMESPACE "CaptureConvertAudioDataTP"
 
 DEFINE_LOG_CATEGORY_STATIC(LogCaptureConvertAudioDataThirdParty, Log, All);
@@ -126,74 +128,25 @@ FCaptureConvertAudioDataThirdParty::FResult FCaptureConvertAudioDataThirdParty::
 	FNamingTokenResultData AudioEncoderCommandResult = NamingTokensSubsystem->EvaluateTokenString(CommandArgs, AudioEncoderTokenArgs);
 	CommandArgs = AudioEncoderCommandResult.EvaluatedText.ToString();
 
-	UE_LOG(LogCaptureConvertAudioDataThirdParty, Display, TEXT("Running the command: %s %s"), *ThirdPartyEncoder.Encoder, *CommandArgs);
-
-	void* ReadPipe = nullptr;
-	void* WritePipe = nullptr;
-	verify(FPlatformProcess::CreatePipe(ReadPipe, WritePipe, false));
-
-	constexpr bool bLaunchDetached = false;
-	constexpr bool bLaunchHidden = true;
-	constexpr bool bLaunchReallyHidden = true;
-	FProcHandle ProcHandle = 
-		FPlatformProcess::CreateProc(*ThirdPartyEncoder.Encoder, *CommandArgs, bLaunchDetached, bLaunchHidden, bLaunchReallyHidden, nullptr, 0, nullptr, WritePipe, nullptr);
-
-	ON_SCOPE_EXIT
-	{
-		if (Params.StopToken.IsStopRequested())
-		{
-			FPlatformProcess::TerminateProc(ProcHandle);
-		}
-
-		FPlatformProcess::CloseProc(ProcHandle);
-		FPlatformProcess::ClosePipe(ReadPipe, WritePipe);
-	};
-
-	if (!ProcHandle.IsValid())
-	{
-		FText Message = FText::Format(LOCTEXT("CaptureConvertAudioNodeTP_ProcessNotFound", "Failed to start the process {0} {1}"),
-									  FText::FromString(ThirdPartyEncoder.Encoder),
-									  FText::FromString(CommandArgs));
-		return MakeError(MoveTemp(Message));
-	}
+	FProcessRunnerResult RunnerResult = FProcessRunner::Run(ThirdPartyEncoder.Encoder, CommandArgs, &Params.StopToken);
 
-	TArray<uint8> FullCommandOutput;
-	while (FPlatformProcess::IsProcRunning(ProcHandle))
+	if (RunnerResult.HasError())
 	{
-		TArray<uint8> CommandOutput = UE::CaptureManager::ReadPipe(ReadPipe);
+		EProcessRunnerError RunnerError = RunnerResult.StealError();
 
-		if (CommandOutput.IsEmpty())
+		FText Message;
+		switch (RunnerError)
 		{
-			FPlatformProcess::Sleep(0.1);
+			case EProcessRunnerError::LaunchFailed:
+				Message = LOCTEXT("CaptureConvertAudioNodeTP_ProcessNotFound", "Failed to start the process");
+				break;
+			case EProcessRunnerError::Cancelled:
+				Message = LOCTEXT("CaptureConvertAudioNodeTP_AbortedByUser", "Aborted by user");
+				break;
+			default:
+				Message = LOCTEXT("CaptureConvertAudioNodeTP_ErrorRunning", "Error occurred while running the third party encoder");
+				break;
 		}
-
-		FullCommandOutput.Append(MoveTemp(CommandOutput));
-
-		if (Params.StopToken.IsStopRequested())
-		{
-			FText Message = LOCTEXT("CaptureConvertAudioNodeTP_AbortedByUser", "Aborted by user");
-			return MakeError(MoveTemp(Message));
-		}
-	}
-
-	int32 ReturnCode = 0;
-	FPlatformProcess::GetProcReturnCode(ProcHandle, &ReturnCode);
-
-	FullCommandOutput.Append(UE::CaptureManager::ReadPipe(ReadPipe));
-	if (ReturnCode != 0)
-	{
-		if (!FullCommandOutput.IsEmpty())
-		{
-			UE_LOG(LogCaptureConvertAudioDataThirdParty, Error,
-				   TEXT("Failed to run the command: %s %s"), *ThirdPartyEncoder.Encoder, *CommandArgs);
-
-			FString CommandOutputStr = FString::ConstructFromPtrSize(reinterpret_cast<const UTF8CHAR*>(FullCommandOutput.GetData()), FullCommandOutput.Num());
-			UE_LOG(LogCaptureConvertAudioDataThirdParty, Display,
-				   TEXT("Output from the command:\n>>>>>>\n%s<<<<<<"), *CommandOutputStr);
-		}
-
-		FText Message = FText::Format(LOCTEXT("CaptureConvertAudioNodeTP_ErrorRunning", "Error while running the third party encoder (ReturnCode={0})"), 
-									  FText::AsNumber(ReturnCode));
 		return MakeError(MoveTemp(Message));
 	}
 
diff --git a/Engine/Plugins/VirtualProduction/CaptureManager/CaptureManagerApp/Source/CaptureDataConverter/Private/Nodes/ThirdPartyEncoder/CaptureConvertVideoDataThirdParty.cpp b/Engine/Plugins/VirtualProduction/CaptureManager/CaptureManagerApp/Source/CaptureDataConverter/Private/Nodes/ThirdPartyEncoder/CaptureConvertVideoDataThirdParty.cpp
--- a/Engine/Plugins/VirtualProduction/CaptureManager/CaptureManagerApp/Source/CaptureDataConverter/Private/Nodes/ThirdPartyEncoder/CaptureConvertVideoDataThirdParty.cpp
+++ b/Engine/Plugins/VirtualProduction/CaptureManager/CaptureManagerApp/Source/CaptureDataConverter/Private/Nodes/ThirdPartyEncoder/CaptureConvertVideoDataThirdParty.cpp
@@ -25,6 +25,8 @@
 
 #include "Internationalization/Regex.h"
 
+#include "ProcessRunner/ProcessRunner.h"
+
 #define LOCTEXT_NAMESPACE "CaptureConvertVideoDataTP"
 
 DEFINE_LOG_CATEGORY_STATIC(LogCaptureConvertVideoDataThirdParty, Log, All);
@@ -232,76 +234,26 @@ FCaptureConvertVideoDataThirdParty::FResult FCaptureConvertVideoDataThirdParty::
 	CommandArgs = FString::Format(*CommandArgs, VideoEncoderFormatArgs);
 	FNamingTokenResultData VideoEncoderCommandResult = NamingTokensSubsystem->EvaluateTokenString(CommandArgs, VideoEncoderTokenArgs);
 	CommandArgs = VideoEncoderCommandResult.EvaluatedText.ToString();
-	
-	UE_LOG(LogCaptureConvertVideoDataThirdParty, Display, TEXT("Running the command: %s %s"), *ThirdPartyEncoder.Encoder, *CommandArgs);
 
-	void* ReadPipe = nullptr;
-	void* WritePipe = nullptr;
-	verify(FPlatformProcess::CreatePipe(ReadPipe, WritePipe, false));
+	FProcessRunnerResult RunnerResult = FProcessRunner::Run(ThirdPartyEncoder.Encoder, CommandArgs, &Params.StopToken);
 
-	constexpr bool bLaunchDetached = false;
-	constexpr bool bLaunchHidden = true;
-	constexpr bool bLaunchReallyHidden = true;
-	FProcHandle ProcHandle = 
-		FPlatformProcess::CreateProc(*ThirdPartyEncoder.Encoder, *CommandArgs, bLaunchDetached, bLaunchHidden, bLaunchReallyHidden, nullptr, 0, nullptr, WritePipe, nullptr);
-
-	ON_SCOPE_EXIT
+	if (RunnerResult.HasError())
 	{
-		if (Params.StopToken.IsStopRequested())
-		{
-			FPlatformProcess::TerminateProc(ProcHandle);
-		}
-
-		FPlatformProcess::CloseProc(ProcHandle);
-		FPlatformProcess::ClosePipe(ReadPipe, WritePipe);
-	};
+		EProcessRunnerError RunnerError = RunnerResult.StealError();
 
-	if (!ProcHandle.IsValid())
-	{
-		FText Message = FText::Format(LOCTEXT("CaptureConvertVideoNodeTP_ProcessNotFound", "Failed to start the process {0} {1}"),
-									  FText::FromString(ThirdPartyEncoder.Encoder),
-									  FText::FromString(CommandArgs));
-		return MakeError(MoveTemp(Message));
-	}
-
-	TArray<uint8> FullCommandOutput;
-	while (FPlatformProcess::IsProcRunning(ProcHandle))
-	{
-		TArray<uint8> CommandOutput = UE::CaptureManager::ReadPipe(ReadPipe);
-
-		if (CommandOutput.IsEmpty())
+		FText Message;
+		switch (RunnerError)
 		{
-			FPlatformProcess::Sleep(0.1);
+			case EProcessRunnerError::LaunchFailed:
+				Message = LOCTEXT("CaptureConvertVideoNodeTP_ProcessNotFound", "Failed to start the process");
+				break;
+			case EProcessRunnerError::Cancelled:
+				Message = LOCTEXT("CaptureConvertVideoNodeTP_AbortedByUser", "Aborted by user");
+				break;
+			default:
+				Message = LOCTEXT("CaptureConvertVideoNodeTP_ErrorRunning", "Error occurred while running the third party encoder.");
+				break;
 		}
-
-		FullCommandOutput.Append(MoveTemp(CommandOutput));
-		
-		if (Params.StopToken.IsStopRequested())
-		{
-			FText Message = LOCTEXT("CaptureConvertVideoNodeTP_AbortedByUser", "Aborted by user");
-			return MakeError(MoveTemp(Message));
-		}
-	}
-
-	int32 ReturnCode = 0;
-	FPlatformProcess::GetProcReturnCode(ProcHandle, &ReturnCode);
-
-	FullCommandOutput.Append(UE::CaptureManager::ReadPipe(ReadPipe));
-
-	if (ReturnCode != 0)
-	{
-		if (!FullCommandOutput.IsEmpty())
-		{
-			UE_LOG(LogCaptureConvertVideoDataThirdParty, Error,
-				   TEXT("Failed to run the command: %s %s"), *ThirdPartyEncoder.Encoder, *CommandArgs);
-
-			FString CommandOutputStr = FString::ConstructFromPtrSize(reinterpret_cast<const UTF8CHAR*>(FullCommandOutput.GetData()), FullCommandOutput.Num());
-			UE_LOG(LogCaptureConvertVideoDataThirdParty, Display,
-				   TEXT("Output from the command:\n>>>>>>\n%s<<<<<<"), *CommandOutputStr);
-		}
-
-		FText Message = FText::Format(LOCTEXT("CaptureConvertVideoNodeTP_ErrorRunning", "Error while running the third party encoder (ReturnCode={0})"),
-									  FText::AsNumber(ReturnCode));
 		return MakeError(MoveTemp(Message));
 	}
 
diff --git a/Engine/Plugins/VirtualProduction/CaptureManager/CaptureManagerApp/Source/IngestLiveLinkDevice/IngestLiveLinkDevice.Build.cs b/Engine/Plugins/VirtualProduction/CaptureManager/CaptureManagerApp/Source/IngestLiveLinkDevice/IngestLiveLinkDevice.Build.cs
--- a/Engine/Plugins/VirtualProduction/CaptureManager/CaptureManagerApp/Source/IngestLiveLinkDevice/IngestLiveLinkDevice.Build.cs
+++ b/Engine/Plugins/VirtualProduction/CaptureManager/CaptureManagerApp/Source/IngestLiveLinkDevice/IngestLiveLinkDevice.Build.cs
@@ -23,6 +23,7 @@ public class IngestLiveLinkDevice : ModuleRules
 			"CaptureDataUtils",
 			"CaptureDataConverter",
 			"CaptureManagerUnrealEndpoint",
+			"CaptureUtils",
 			"Engine",
 			"Json",
 			"LiveLinkDevice",
diff --git a/Engine/Plugins/VirtualProduction/CaptureManager/CaptureManagerApp/Source/IngestLiveLinkDevice/Private/Utils/CaptureExtractTimecode.cpp b/Engine/Plugins/VirtualProduction/CaptureManager/CaptureManagerApp/Source/IngestLiveLinkDevice/Private/Utils/CaptureExtractTimecode.cpp
--- a/Engine/Plugins/VirtualProduction/CaptureManager/CaptureManagerApp/Source/IngestLiveLinkDevice/Private/Utils/CaptureExtractTimecode.cpp
+++ b/Engine/Plugins/VirtualProduction/CaptureManager/CaptureManagerApp/Source/IngestLiveLinkDevice/Private/Utils/CaptureExtractTimecode.cpp
@@ -18,6 +18,8 @@
 
 #include "ParseTakeUtils.h"
 
+#include "ProcessRunner/ProcessRunner.h"
+
 #define LOCTEXT_NAMESPACE "CaptureExtractTimecode"
 
 DEFINE_LOG_CATEGORY_STATIC(LogCaptureExtractTimecode, Log, All);
@@ -293,89 +295,18 @@ FCaptureExtractVideoInfo::FExtractResult FCaptureExtractVideoInfo::ExtractInfoUs
 
 	FString CommandArgs = FString::Format(TEXT("-v error -select_streams v:0 -show_entries format_tags=timecode:stream=r_frame_rate -of default=noprint_wrappers=1:nokey=1 \"{0}\""), { InFilePath });
 
-	void* ReadPipe = nullptr;
-	void* WritePipe = nullptr;
-	verify(FPlatformProcess::CreatePipe(ReadPipe, WritePipe, false));
-
-	constexpr bool bLaunchDetached = false;
-	constexpr bool bLaunchHidden = true;
-	constexpr bool bLaunchReallyHidden = true;
-	FProcHandle ProcHandle =
-		FPlatformProcess::CreateProc(*InFFProbePath, *CommandArgs, bLaunchDetached, bLaunchHidden, bLaunchReallyHidden, nullptr, 0, nullptr, WritePipe, nullptr);
-
-	ON_SCOPE_EXIT
-	{
-		FPlatformProcess::TerminateProc(ProcHandle);
-		FPlatformProcess::CloseProc(ProcHandle);
-		FPlatformProcess::ClosePipe(ReadPipe, WritePipe);
-	};
-
-	if (!ProcHandle.IsValid())
+	FProcessRunnerResult RunnerResult = FProcessRunner::Run(InFFProbePath, CommandArgs, nullptr, TimeoutPeriod);
+	
+	if (!RunnerResult.HasValue())
 	{
-		UE_LOG(LogCaptureExtractTimecode, Display, TEXT("Extract timecode: Failed to start the process %s %s"), *InFFProbePath, *CommandArgs);
 		return MakeError(ECaptureExtractInfoError::InternalError);
 	}
 
-	const FDateTime WaitStart = FDateTime::Now();
-
-	TArray<uint8> FullCommandOutput;
-	while (FPlatformProcess::IsProcRunning(ProcHandle))
-	{
-		TArray<uint8> CommandOutput;
-		{
-			bool Read = FPlatformProcess::ReadPipeToArray(ReadPipe, CommandOutput);
-			if (!Read)
-			{
-				CommandOutput.Empty();
-			}
-		}
-
-		if ((FDateTime::Now() - WaitStart).GetSeconds() > TimeoutPeriod)
-		{
-			// Timed out
-			break;
-		}
-
-		if (CommandOutput.IsEmpty())
-		{
-			FPlatformProcess::Sleep(0.1f);
-			continue;
-		}
-
-		FullCommandOutput.Append(MoveTemp(CommandOutput));
-
-	}
-
-	int32 ReturnCode = 0;
-	FPlatformProcess::GetProcReturnCode(ProcHandle, &ReturnCode);
-
-	TArray<uint8> CommandOutput;
-	{
-		bool Read = FPlatformProcess::ReadPipeToArray(ReadPipe, CommandOutput);
-		if (!Read)
-		{
-			CommandOutput.Empty();
-		}
-	}
-	FullCommandOutput.Append(CommandOutput);
-
-	if (ReturnCode != 0)
-	{
-		if (!FullCommandOutput.IsEmpty())
-		{
-			UE_LOG(LogCaptureExtractTimecode, Display, TEXT("Failed to run the command: %s %s"), *InFFProbePath, *CommandArgs);
-
-			FString CommandOutputStr = FString::ConstructFromPtrSize(reinterpret_cast<const UTF8CHAR*>(FullCommandOutput.GetData()), FullCommandOutput.Num());
-			UE_LOG(LogCaptureExtractTimecode, Display,
-				   TEXT("Output from the command:\n>>>>>>\n%s<<<<<<"), *CommandOutputStr);
-		}
-
-		return MakeError(ECaptureExtractInfoError::InternalError);
-	}
 	bool bFoundTimecode = false;
 	bool bFoundFrameRate = false;
 
-	FString CommandOutputStr = FString::ConstructFromPtrSize(reinterpret_cast<const UTF8CHAR*>(FullCommandOutput.GetData()), FullCommandOutput.Num());
+	TArray<uint8> CommandOutput = RunnerResult.StealValue();
+	FString CommandOutputStr = FString::ConstructFromPtrSize(reinterpret_cast<const UTF8CHAR*>(CommandOutput.GetData()), CommandOutput.Num());
 
 	// Split the timecode and frame rate lines from the output
 	TArray<FString> OutputLinesStr;
diff --git a/Engine/Plugins/VirtualProduction/CaptureManager/CaptureManagerCore/Source/CaptureUtils/Internal/ProcessRunner/ProcessRunner.h b/Engine/Plugins/VirtualProduction/CaptureManager/CaptureManagerCore/Source/CaptureUtils/Internal/ProcessRunner/ProcessRunner.h
new file mode 100644
--- /dev/null
+++ b/Engine/Plugins/VirtualProduction/CaptureManager/CaptureManagerCore/Source/CaptureUtils/Internal/ProcessRunner/ProcessRunner.h
@@ -0,0 +1,39 @@
+// Copyright Epic Games, Inc. All Rights Reserved.
+
+#pragma once
+
+#include "CoreMinimal.h"
+#include "Templates/ValueOrError.h"
+
+namespace UE::CaptureManager
+{
+
+class FStopToken;
+
+enum class EProcessRunnerError : int32
+{
+	PipeCreationFailed,
+	LaunchFailed,
+	Timeout,
+	Cancelled,
+	FailedToGetReturnCode,
+	ExitedWithError
+};
+
+using FProcessRunnerResult = TValueOrError<TArray<uint8>, EProcessRunnerError>;
+
+class UE_INTERNAL CAPTUREUTILS_API FProcessRunner
+{
+public:
+	static FProcessRunnerResult Run(
+		const FString& InProcessPath, 
+		const FString& InProcessArguments, 
+		const FStopToken* InStopToken = nullptr, 
+		const TOptional<int32> InTimeoutSeconds = TOptional<int32>());
+
+private:
+	static TArray<uint8> ReadPipe(void* InPipe);
+	static TArray<uint8> DrainPipe(void* InPipe);
+};
+
+}
diff --git a/Engine/Plugins/VirtualProduction/CaptureManager/CaptureManagerCore/Source/CaptureUtils/Private/ProcessRunner/ProcessRunner.cpp b/Engine/Plugins/VirtualProduction/CaptureManager/CaptureManagerCore/Source/CaptureUtils/Private/ProcessRunner/ProcessRunner.cpp
new file mode 100644
--- /dev/null
+++ b/Engine/Plugins/VirtualProduction/CaptureManager/CaptureManagerCore/Source/CaptureUtils/Private/ProcessRunner/ProcessRunner.cpp
@@ -0,0 +1,178 @@
+// Copyright Epic Games, Inc. All Rights Reserved.
+
+#include "ProcessRunner/ProcessRunner.h"
+
+#include "Async/StopToken.h"
+#include "HAL/PlatformProcess.h"
+#include "Misc/Paths.h"
+#include "Misc/ScopeExit.h"
+
+DEFINE_LOG_CATEGORY_STATIC(LogCaptureManagerProcessRunner, Log, All);
+
+namespace UE::CaptureManager
+{
+
+FProcessRunnerResult FProcessRunner::Run(
+	const FString& InProcessPath, 
+	const FString& InProcessArguments, 
+	const FStopToken* InStopToken, 
+	const TOptional<int32> InTimeoutSeconds)
+{
+	UE_LOG(LogCaptureManagerProcessRunner, Log, TEXT("Executing \"%s\" with arguments \"%s\""), *FPaths::GetCleanFilename(InProcessPath), *InProcessArguments);
+
+	void* StdOutReadPipe = nullptr;
+	void* StdOutWritePipe = nullptr;
+
+	void* StdErrReadPipe = nullptr;
+	void* StdErrWritePipe = nullptr;
+
+	void* StdInReadPipe = nullptr;
+	void* StdInWritePipe = nullptr;
+
+	ON_SCOPE_EXIT
+	{
+		FPlatformProcess::ClosePipe(StdOutReadPipe, StdOutWritePipe);
+		FPlatformProcess::ClosePipe(StdErrReadPipe, StdErrWritePipe);
+		FPlatformProcess::ClosePipe(StdInReadPipe, StdInWritePipe);
+	};
+
+	/* Create the read/write pipes for process I/O
+	 * - StdOut: Parent process owns the read pipe, child owns write.
+	 * - StdErr: Same as StdOut
+	 * - StdIn: Parent process owns the write pipe, child owns read.
+	 * For StdIn we must specify 'true' for the third argument (bWritePipeLocal) in order
+	 * to swap ownership so that the child can inherit the read pipe.
+	*/
+	if (!FPlatformProcess::CreatePipe(StdOutReadPipe, StdOutWritePipe, false)
+		|| !FPlatformProcess::CreatePipe(StdErrReadPipe, StdErrWritePipe, false)
+		|| !FPlatformProcess::CreatePipe(StdInReadPipe, StdInWritePipe, true))
+	{
+		return MakeError(EProcessRunnerError::PipeCreationFailed);
+	}
+
+	constexpr bool bLaunchDetached = false;
+	constexpr bool bLaunchHidden = true;
+	constexpr bool bLaunchReallyHidden = true;
+	FProcHandle ProcHandle = FPlatformProcess::CreateProc(
+		*InProcessPath, 
+		*InProcessArguments, 
+		bLaunchDetached, 
+		bLaunchHidden, 
+		bLaunchReallyHidden, 
+		nullptr, 
+		0, 
+		nullptr, 
+		StdOutWritePipe, 
+		StdInReadPipe, 
+		StdErrWritePipe);
+
+	bool bTerminateOnScopeExit = false;
+
+	ON_SCOPE_EXIT
+	{
+		if (bTerminateOnScopeExit)
+		{
+			FPlatformProcess::TerminateProc(ProcHandle);
+		}
+		FPlatformProcess::CloseProc(ProcHandle);
+	};
+
+	if (!ProcHandle.IsValid())
+	{
+		UE_LOG(LogCaptureManagerProcessRunner, Error, TEXT("Failed to launch \"%s\" with arguments \"%s\""), *FPaths::GetCleanFilename(InProcessPath), *InProcessArguments);
+		return MakeError(EProcessRunnerError::LaunchFailed);
+	}
+
+	const FDateTime StartDateTime = FDateTime::Now();
+
+	TArray<uint8> StandardOutput;
+	TArray<uint8> StandardError;
+	while (FPlatformProcess::IsProcRunning(ProcHandle))
+	{
+		TArray<uint8> StandardOutputFragment = ReadPipe(StdOutReadPipe);
+		TArray<uint8> StandardErrorFragment = ReadPipe(StdErrReadPipe);
+
+		StandardOutput.Append(MoveTemp(StandardOutputFragment));
+		StandardError.Append(MoveTemp(StandardErrorFragment));
+
+		if (StandardOutputFragment.IsEmpty() && StandardErrorFragment.IsEmpty())
+		{
+			FPlatformProcess::Sleep(0.01f);
+		}
+
+		if (InTimeoutSeconds.IsSet())
+		{
+			int32 TimeoutSeconds = InTimeoutSeconds.GetValue();
+			if ((FDateTime::Now() - StartDateTime).GetSeconds() > TimeoutSeconds)
+			{
+				bTerminateOnScopeExit = true;
+				UE_LOG(LogCaptureManagerProcessRunner, Error, TEXT("\"%s\" timed out after %d seconds"), *FPaths::GetCleanFilename(InProcessPath), TimeoutSeconds);
+				return MakeError(EProcessRunnerError::Timeout);
+			}
+		}
+
+		if (InStopToken && InStopToken->IsStopRequested())
+		{
+			bTerminateOnScopeExit = true;
+			UE_LOG(LogCaptureManagerProcessRunner, Warning, TEXT("\"%s\" cancelled by the user"), *FPaths::GetCleanFilename(InProcessPath));
+			return MakeError(EProcessRunnerError::Cancelled);
+		}
+	}
+
+	// Drain pipes of all remaining data
+	StandardOutput.Append(DrainPipe(StdOutReadPipe));
+	StandardError.Append(DrainPipe(StdErrReadPipe));
+
+	int32 ReturnCode = 0;
+	if (!FPlatformProcess::GetProcReturnCode(ProcHandle, &ReturnCode))
+	{
+		UE_LOG(LogCaptureManagerProcessRunner, Error, TEXT("Failed to get return code when running \"%s\""), *FPaths::GetCleanFilename(InProcessPath))
+		return MakeError(EProcessRunnerError::FailedToGetReturnCode);
+	}
+
+	if (ReturnCode != 0)
+	{
+		UE_LOG(LogCaptureManagerProcessRunner, Error, TEXT("\"%s\" exited with error code %d"), *FPaths::GetCleanFilename(InProcessPath), ReturnCode);
+
+		if (!StandardOutput.IsEmpty())
+		{
+			FString StandardOutputString = FString::ConstructFromPtrSize(reinterpret_cast<const UTF8CHAR*>(StandardOutput.GetData()), StandardOutput.Num());
+			UE_LOG(LogCaptureManagerProcessRunner, Log, TEXT("Standard output:\n>>>>>>\n%s<<<<<<"), *StandardOutputString);
+		}
+
+		if (!StandardError.IsEmpty())
+		{
+			FString StandardErrorString = FString::ConstructFromPtrSize(reinterpret_cast<const UTF8CHAR*>(StandardError.GetData()), StandardError.Num());
+			UE_LOG(LogCaptureManagerProcessRunner, Log, TEXT("Error output:\n>>>>>>\n%s<<<<<<\n"), *StandardErrorString);
+		}
+
+		return MakeError(EProcessRunnerError::ExitedWithError);
+	}
+
+	return MakeValue(StandardOutput);
+}
+
+TArray<uint8> FProcessRunner::ReadPipe(void* InPipe)
+{
+	TArray<uint8> Output;
+	bool Read = FPlatformProcess::ReadPipeToArray(InPipe, Output);
+	if (!Read)
+	{
+		Output.Empty();
+	}
+	return Output;
+}
+
+TArray<uint8> FProcessRunner::DrainPipe(void* InPipe)
+{
+	TArray<uint8> Drained;
+	TArray<uint8> Fragment;
+	do 
+	{
+		Fragment = ReadPipe(InPipe);
+		Drained.Append(Fragment);
+	} while (!Fragment.IsEmpty());
+	return Drained;
+}
+
+}
diff --git a/Engine/Plugins/VirtualProduction/CaptureManager/CaptureManagerDevices/CaptureManagerDevices.uplugin b/Engine/Plugins/VirtualProduction/CaptureManager/CaptureManagerDevices/CaptureManagerDevices.uplugin
--- a/Engine/Plugins/VirtualProduction/CaptureManager/CaptureManagerDevices/CaptureManagerDevices.uplugin
+++ b/Engine/Plugins/VirtualProduction/CaptureManager/CaptureManagerDevices/CaptureManagerDevices.uplugin
@@ -47,6 +47,10 @@
 		{
 			"Name": "CaptureManagerApp",
 			"Enabled": true
+		},
+		{
+			"Name": "CaptureManagerCore",
+			"Enabled": true
 		}
 	]
 }
\ No newline at end of file
diff --git a/Engine/Plugins/VirtualProduction/CaptureManager/CaptureManagerDevices/Source/VideoLiveLinkDeviceCommon/Private/Utils/VideoDeviceThumbnailExtractor.cpp b/Engine/Plugins/VirtualProduction/CaptureManager/CaptureManagerDevices/Source/VideoLiveLinkDeviceCommon/Private/Utils/VideoDeviceThumbnailExtractor.cpp
--- a/Engine/Plugins/VirtualProduction/CaptureManager/CaptureManagerDevices/Source/VideoLiveLinkDeviceCommon/Private/Utils/VideoDeviceThumbnailExtractor.cpp
+++ b/Engine/Plugins/VirtualProduction/CaptureManager/CaptureManagerDevices/Source/VideoLiveLinkDeviceCommon/Private/Utils/VideoDeviceThumbnailExtractor.cpp
@@ -9,6 +9,7 @@
 #include "Serialization/JsonSerializer.h"
 
 #include "Settings/CaptureManagerSettings.h"
+#include "ProcessRunner/ProcessRunner.h"
 
 #include "Utils/MediaPixelFormatConversions.h"
 
@@ -177,7 +178,14 @@ TArray<FColor> FVideoDeviceThumbnailExtractor::ObtainThumbnailFromThirdPartyEnco
 
 	FString CommandLineArgs = FString::Format(TEXT("-hide_banner -loglevel error -i \"{0}\" -map v:0 -vf \"thumbnail=2, {1}\" -frames:v 1 -f rawvideo -pix_fmt rgb24 -"), { InCurrentFile, ScaleArg });
 
-	TArray<uint8> ThumbnailRawData = RunProcess(InEncoderPath, CommandLineArgs);
+	FProcessRunnerResult RunnerResult = FProcessRunner::Run(InEncoderPath, CommandLineArgs);
+
+	if (!RunnerResult.HasValue())
+	{
+		return TArray<FColor>(); 
+	}
+
+	TArray<uint8> ThumbnailRawData = RunnerResult.StealValue();
 	int32 ThumbnailNumBytes = ThumbnailRawData.Num();
 
 	constexpr int32 NumChannels = 3;
@@ -198,72 +206,4 @@ TArray<FColor> FVideoDeviceThumbnailExtractor::ObtainThumbnailFromThirdPartyEnco
 	return ThumbnailRawColorData;
 }
 
-TArray<uint8> FVideoDeviceThumbnailExtractor::RunProcess(const FString& InProcessName, const FString& InProcessArgs)
-{
-	void* ReadPipe = nullptr;
-	void* WritePipe = nullptr;
-	verify(FPlatformProcess::CreatePipe(ReadPipe, WritePipe, false));
-
-	constexpr bool bLaunchDetached = false;
-	constexpr bool bLaunchHidden = true;
-	constexpr bool bLaunchReallyHidden = true;
-	FProcHandle ProcHandle =
-		FPlatformProcess::CreateProc(*InProcessName, *InProcessArgs, bLaunchDetached, bLaunchHidden, bLaunchReallyHidden, nullptr, 0, nullptr, WritePipe, nullptr);
-
-	ON_SCOPE_EXIT
-	{
-		FPlatformProcess::TerminateProc(ProcHandle);
-		FPlatformProcess::CloseProc(ProcHandle);
-		FPlatformProcess::ClosePipe(ReadPipe, WritePipe);
-	};
-
-	if (!ProcHandle.IsValid())
-	{
-		return TArray<uint8>();
-	}
-
-	TArray<uint8> FullCommandOutput;
-	while (FPlatformProcess::IsProcRunning(ProcHandle))
-	{
-		TArray<uint8> CommandOutput = ReadFromPipe(ReadPipe);
-
-		if (CommandOutput.IsEmpty())
-		{
-			FPlatformProcess::Sleep(0.010f);
-		}
-
-		FullCommandOutput.Append(MoveTemp(CommandOutput));
-	}
-
-	int32 ReturnCode = 0;
-	FPlatformProcess::GetProcReturnCode(ProcHandle, &ReturnCode);
-
-	if (ReturnCode != 0)
-	{
-		return TArray<uint8>();
-	}
-
-	// Drain pipe of all remaining data
-	TArray<uint8> CommandOutput;
-	do 
-	{
-		CommandOutput = ReadFromPipe(ReadPipe);
-		FullCommandOutput.Append(CommandOutput);
-	} while (!CommandOutput.IsEmpty());
-
-	return FullCommandOutput;
-}
-
-TArray<uint8> FVideoDeviceThumbnailExtractor::ReadFromPipe(void* InPipe)
-{
-	TArray<uint8> CommandOutput;
-	bool Read = FPlatformProcess::ReadPipeToArray(InPipe, CommandOutput);
-	if (!Read)
-	{
-		CommandOutput.Empty();
-	}
-
-	return CommandOutput;
-}
-
 }
diff --git a/Engine/Plugins/VirtualProduction/CaptureManager/CaptureManagerDevices/Source/VideoLiveLinkDeviceCommon/VideoLiveLinkDeviceCommon.Build.cs b/Engine/Plugins/VirtualProduction/CaptureManager/CaptureManagerDevices/Source/VideoLiveLinkDeviceCommon/VideoLiveLinkDeviceCommon.Build.cs
--- a/Engine/Plugins/VirtualProduction/CaptureManager/CaptureManagerDevices/Source/VideoLiveLinkDeviceCommon/VideoLiveLinkDeviceCommon.Build.cs
+++ b/Engine/Plugins/VirtualProduction/CaptureManager/CaptureManagerDevices/Source/VideoLiveLinkDeviceCommon/VideoLiveLinkDeviceCommon.Build.cs
@@ -21,6 +21,7 @@ public class VideoLiveLinkDeviceCommon : ModuleRules
 		{
 			"CaptureManagerMediaRW",
 			"CaptureManagerSettings",
+			"CaptureUtils",
 			"Engine",
 			"Slate",
 			"SlateCore",
