Extend CPS with capabilities feature.

#Synced-CL 49723046

diff --git a/Engine/Plugins/VirtualProduction/CaptureManager/CaptureManagerCore/Source/CaptureProtocolStack/Private/Control/Capabilities/Utils.cpp b/Engine/Plugins/VirtualProduction/CaptureManager/CaptureManagerCore/Source/CaptureProtocolStack/Private/Control/Capabilities/Utils.cpp
new file mode 100644
--- /dev/null
+++ b/Engine/Plugins/VirtualProduction/CaptureManager/CaptureManagerCore/Source/CaptureProtocolStack/Private/Control/Capabilities/Utils.cpp
@@ -0,0 +1,131 @@
+// Copyright Epic Games, Inc. All Rights Reserved.
+
+#include "Control/Capabilities/Utils.h"
+#include "Control/Messages/ControlJsonUtilities.h"
+#include "Control/Messages/Constants.h"
+#include "Control/Messages/ProtocolConstants.h"
+
+#include "Dom/JsonValue.h"
+
+namespace UE::CaptureManager
+{
+TProtocolResult<void> FCapabilityUtilities::TypeStringToEnum(const FString& InTypeString, ECapabilityValue::Type& OutType)
+{
+	if (InTypeString == CPS::Capabilities::GBoolean)
+	{
+		OutType = ECapabilityValue::Type::Boolean;
+	}
+	else if (InTypeString == CPS::Capabilities::GInteger)
+	{
+		OutType = ECapabilityValue::Type::Integer;
+	}
+	else if (InTypeString == CPS::Capabilities::GFloat)
+	{
+		OutType = ECapabilityValue::Type::Float;
+	}
+	else if (InTypeString == CPS::Capabilities::GString)
+	{
+		OutType = ECapabilityValue::Type::String;
+	}
+	else if (InTypeString == CPS::Capabilities::GEnumeration)
+	{
+		OutType = ECapabilityValue::Type::Enum;
+	}
+	else
+	{
+		return FCaptureProtocolError(TEXT("Invalud value type string."));
+	}
+	return ResultOk;
+}
+
+TProtocolResult<void> FCapabilityUtilities::ValueFromJsonObject(const TSharedPtr<FJsonObject>& InJsonObject,
+																const FString& InKey,
+																const ECapabilityValue::Type InType,
+																FCapabilityValue& OutValue)
+{
+	if (InType == ECapabilityValue::Type::Boolean)
+	{
+		bool Value;
+		CHECK_PARSE(FJsonUtility::ParseBool(InJsonObject, InKey, Value));
+		OutValue.Set<bool>(Value);
+	}
+	else if (InType == ECapabilityValue::Type::Integer)
+	{
+		int Value;
+		CHECK_PARSE(FJsonUtility::ParseNumber(InJsonObject, InKey, Value));
+		OutValue.Set<int>(Value);
+	}
+	else if (InType == ECapabilityValue::Type::Float)
+	{
+		float Value;
+		CHECK_PARSE(FJsonUtility::ParseNumber(InJsonObject, InKey, Value));
+		OutValue.Set<float>(Value);
+	}
+	else if ((InType == ECapabilityValue::Type::String) || (InType == ECapabilityValue::Type::Enum))
+	{
+		FString Value;
+		CHECK_PARSE(FJsonUtility::ParseString(InJsonObject, InKey, Value));
+		OutValue.Set<FString>(Value);
+	}
+	else
+	{
+		check(false);
+	}
+
+	return ResultOk;
+}
+
+void FCapabilityUtilities::SetProperJsonFieldType(const FString& InKey,
+												  const FCapabilityValue& InValue,
+												  TSharedPtr<FJsonObject>& OutJsonObjectPtr)
+{
+	ECapabilityValue::Type Type = TypeFromVariant(InValue);
+
+	if (Type == ECapabilityValue::Type::Boolean)
+	{
+		OutJsonObjectPtr->SetBoolField(InKey, InValue.Get<bool>());
+	}
+	else if (Type == ECapabilityValue::Type::Integer)
+	{
+		OutJsonObjectPtr->SetNumberField(InKey, InValue.Get<int>());
+	}
+	else if (Type == ECapabilityValue::Type::Float)
+	{
+		OutJsonObjectPtr->SetNumberField(InKey, InValue.Get<float>());
+	}
+	else if ((Type == ECapabilityValue::Type::String) || (Type == ECapabilityValue::Type::Enum))
+	{
+		OutJsonObjectPtr->SetStringField(InKey, InValue.Get<FString>());
+	}
+	else
+	{
+		check(false);
+	}
+}
+
+ECapabilityValue::Type FCapabilityUtilities::TypeFromVariant(const FCapabilityValue& InValue)
+{
+	if (InValue.GetIndex() == 0)
+	{
+		return ECapabilityValue::Type::Boolean;
+	}
+	else if (InValue.GetIndex() == 1)
+	{
+		return ECapabilityValue::Type::Integer;
+	}
+	else if (InValue.GetIndex() == 2)
+	{
+		return ECapabilityValue::Type::Float;
+	}
+	else if (InValue.GetIndex() == 3)
+	{
+		return ECapabilityValue::Type::String;
+	}
+	else
+	{
+		check(false);
+	}
+	return ECapabilityValue::Type::Boolean;
+}
+
+} // namespace UE::CaptureManager
diff --git a/Engine/Plugins/VirtualProduction/CaptureManager/CaptureManagerCore/Source/CaptureProtocolStack/Private/Control/Capabilities/Utils.h b/Engine/Plugins/VirtualProduction/CaptureManager/CaptureManagerCore/Source/CaptureProtocolStack/Private/Control/Capabilities/Utils.h
new file mode 100644
--- /dev/null
+++ b/Engine/Plugins/VirtualProduction/CaptureManager/CaptureManagerCore/Source/CaptureProtocolStack/Private/Control/Capabilities/Utils.h
@@ -0,0 +1,28 @@
+// Copyright Epic Games, Inc. All Rights Reserved.
+
+#pragma once
+#include "Utility/Error.h"
+#include "Control/Capabilities/Capability.h"
+
+#include "Dom/JsonObject.h"
+
+class FString;
+
+#define CPS_API CAPTUREPROTOCOLSTACK_API
+
+namespace UE::CaptureManager
+{
+struct FCapabilityUtilities
+{
+	static CPS_API TProtocolResult<void> TypeStringToEnum(const FString& InTypeString, ECapabilityValue::Type& OutType);
+	static CPS_API TProtocolResult<void> ValueFromJsonObject(const TSharedPtr<FJsonObject>& InJsonObject,
+															 const FString& InKey,
+															 const ECapabilityValue::Type InType,
+															 FCapabilityValue& OutValue);
+	static CPS_API void SetProperJsonFieldType(const FString& InKey,
+											   const FCapabilityValue& InValue,
+											   TSharedPtr<FJsonObject>& OutJsonObjectPtr);
+
+	static CPS_API ECapabilityValue::Type TypeFromVariant(const FCapabilityValue& InValue);
+};
+} // namespace UE::CaptureManager
diff --git a/Engine/Plugins/VirtualProduction/CaptureManager/CaptureManagerCore/Source/CaptureProtocolStack/Private/Control/Messages/Constants.cpp b/Engine/Plugins/VirtualProduction/CaptureManager/CaptureManagerCore/Source/CaptureProtocolStack/Private/Control/Messages/Constants.cpp
--- a/Engine/Plugins/VirtualProduction/CaptureManager/CaptureManagerCore/Source/CaptureProtocolStack/Private/Control/Messages/Constants.cpp
+++ b/Engine/Plugins/VirtualProduction/CaptureManager/CaptureManagerCore/Source/CaptureProtocolStack/Private/Control/Messages/Constants.cpp
@@ -33,6 +33,18 @@ const TCHAR* GTakeAdded = TEXT("/take/added");
 const TCHAR* GTakeRemoved = TEXT("/take/removed");
 const TCHAR* GTakeUpdated = TEXT("/take/updated");
 
+// Capabilities
+const TCHAR* GCapabilitiesRoot = TEXT("/capabilities");
+const TCHAR* GCapabilitiesListAll = TEXT("list-all");
+const TCHAR* GCapabilityProperties = TEXT("properties");
+const TCHAR* GCapabilityPropertySet = TEXT("set");
+const TCHAR* GCapabilityPropertyGet = TEXT("get");
+const TCHAR* GCapabilityCommands = TEXT("commands");
+const TCHAR* GCapabilityCommandExecute = TEXT("execute");
+const TCHAR* GCapabilityEvents = TEXT("events");
+const TCHAR* GCapabilityEventSubscribe = TEXT("subscribe");
+const TCHAR* GCapabilityEventUnsubscribe = TEXT("unsubscribe");
+
 // Platform Specific: IOS
 const TCHAR* GDiskCapacity = TEXT("state/ios/disk-capacity");
 const TCHAR* GBattery = TEXT("state/ios/battery");
@@ -133,4 +145,37 @@ const TCHAR* GRequestParsingFailed = TEXT("parsingRequestFailed");
 const TCHAR* GNotSupported = TEXT("notSupported");
 const TCHAR* GNoSession = TEXT("noSession");
 const TCHAR* GMissingBody = TEXT("missingBody");
-}
\ No newline at end of file
+} // namespace UE::CaptureManager::CPS::ErrorNames
+
+namespace UE::CaptureManager::CPS::Capabilities
+{
+const TCHAR* GId = TEXT("id");
+const TCHAR* GName = TEXT("name");
+const TCHAR* GValue = TEXT("value");
+const TCHAR* GProperties = TEXT("properties");
+const TCHAR* GType = TEXT("type");
+const TCHAR* GAccess = TEXT("access");
+const TCHAR* GReadOnly = TEXT("readOnly");
+const TCHAR* GReadWrite = TEXT("readWrite");
+const TCHAR* GEnumOptions = TEXT("enumOptions");
+const TCHAR* GCurrentValue = TEXT("currentValue");
+const TCHAR* GMin = TEXT("min");
+const TCHAR* GMax = TEXT("max");
+const TCHAR* GGet = TEXT("get");
+const TCHAR* GSet = TEXT("set");
+const TCHAR* GCommands = TEXT("commands");
+const TCHAR* GParameters = TEXT("parameters");
+const TCHAR* GExecute = TEXT("execute");
+const TCHAR* GOptional = TEXT("optional");
+const TCHAR* GReturnType = TEXT("returnType");
+const TCHAR* GDefaultValue = TEXT("defaultValue");
+const TCHAR* GEvents = TEXT("events");
+const TCHAR* GArguments = TEXT("arguments");
+const TCHAR* GSubscribe = TEXT("subscribe");
+const TCHAR* GUnsubscribe = TEXT("unsubscribe");
+const TCHAR* GBoolean = TEXT("boolean");
+const TCHAR* GInteger = TEXT("integer");
+const TCHAR* GFloat = TEXT("float");
+const TCHAR* GString = TEXT("string");
+const TCHAR* GEnumeration = TEXT("enumeration");
+} // namespace UE::CaptureManager::CPS::Capabilities
diff --git a/Engine/Plugins/VirtualProduction/CaptureManager/CaptureManagerCore/Source/CaptureProtocolStack/Private/Control/Messages/ControlRequest.cpp b/Engine/Plugins/VirtualProduction/CaptureManager/CaptureManagerCore/Source/CaptureProtocolStack/Private/Control/Messages/ControlRequest.cpp
--- a/Engine/Plugins/VirtualProduction/CaptureManager/CaptureManagerCore/Source/CaptureProtocolStack/Private/Control/Messages/ControlRequest.cpp
+++ b/Engine/Plugins/VirtualProduction/CaptureManager/CaptureManagerCore/Source/CaptureProtocolStack/Private/Control/Messages/ControlRequest.cpp
@@ -2,6 +2,10 @@
 
 #include "Control/Messages/ControlRequest.h"
 #include "Control/Messages/Constants.h"
+#include "Control/Messages/ProtocolConstants.h"
+
+#include "Misc/Paths.h"
+#include "Control/Capabilities/Utils.h"
 
 namespace UE::CaptureManager
 {
@@ -193,4 +197,81 @@ TSharedPtr<FJsonObject> FStopStreamingRequest::GetBody() const
 	return Body;
 }
 
-}
\ No newline at end of file
+FListAllCapabilitiesRequest::FListAllCapabilitiesRequest()
+	: FControlRequest(FPaths::Combine(CPS::AddressPaths::GCapabilitiesRoot, CPS::AddressPaths::GCapabilitiesListAll))
+{
+}
+
+FGetCapabilityPropertyRequest::FGetCapabilityPropertyRequest(FString InCapabilityId, FString InPropertyId)
+	: FControlRequest(FPaths::Combine(CPS::AddressPaths::GCapabilitiesRoot,
+									  MoveTemp(InCapabilityId),
+									  CPS::AddressPaths::GCapabilityProperties,
+									  MoveTemp(InPropertyId),
+									  CPS::AddressPaths::GCapabilityPropertyGet))
+{
+}
+
+FSetCapabilityPropertyRequest::FSetCapabilityPropertyRequest(FString InCapabilityId,
+															 FString InPropertyId,
+															 FCapabilityValue InValue)
+	: FControlRequest(FPaths::Combine(CPS::AddressPaths::GCapabilitiesRoot,
+									  MoveTemp(InCapabilityId),
+									  CPS::AddressPaths::GCapabilityProperties,
+									  MoveTemp(InPropertyId),
+									  CPS::AddressPaths::GCapabilityPropertySet))
+{
+	Value = MoveTemp(InValue);
+}
+
+TSharedPtr<FJsonObject> FSetCapabilityPropertyRequest::GetBody() const
+{
+	TSharedPtr<FJsonObject> Body = MakeShared<FJsonObject>();
+	FCapabilityUtilities::SetProperJsonFieldType(CPS::Capabilities::GValue, Value, Body);
+	return Body;
+}
+
+FExecuteCapabilityCommandRequest::FExecuteCapabilityCommandRequest(FString InCapabilityId,
+																   FString InCommandId,
+																   FCapabilityValues InArguments)
+	: FControlRequest(FPaths::Combine(CPS::AddressPaths::GCapabilitiesRoot,
+									  MoveTemp(InCapabilityId),
+									  CPS::AddressPaths::GCapabilityCommands,
+									  MoveTemp(InCommandId),
+									  CPS::AddressPaths::GCapabilityCommandExecute))
+{
+	Arguments = MoveTemp(InArguments);
+}
+
+TSharedPtr<FJsonObject> FExecuteCapabilityCommandRequest::GetBody() const
+{
+	if (Arguments.Num() > 0)
+	{
+		TSharedPtr<FJsonObject> Body = MakeShared<FJsonObject>();
+		TSharedPtr<FJsonObject> JsonArguments = MakeShared<FJsonObject>();
+		for (const TPair<FString, FCapabilityValue>& Argument : Arguments)
+		{
+			FCapabilityUtilities::SetProperJsonFieldType(Argument.Key, Argument.Value, JsonArguments);
+		}
+		Body->SetObjectField(CPS::Capabilities::GArguments, JsonArguments);
+		return Body;
+	}
+	return nullptr;
+}
+
+FCapabilitySubscribeRequest::FCapabilitySubscribeRequest(FString InCapabilityId)
+	: FControlRequest(FPaths::Combine(CPS::AddressPaths::GCapabilitiesRoot,
+									  InCapabilityId,
+									  CPS::AddressPaths::GCapabilityEvents,
+									  CPS::AddressPaths::GCapabilityEventSubscribe))
+{
+}
+
+FCapabilityUnsubscribeRequest::FCapabilityUnsubscribeRequest(FString InCapabilityId)
+	: FControlRequest(FPaths::Combine(CPS::AddressPaths::GCapabilitiesRoot,
+									  InCapabilityId,
+									  CPS::AddressPaths::GCapabilityEvents,
+									  CPS::AddressPaths::GCapabilityEventUnsubscribe))
+{
+}
+
+} // namespace UE::CaptureManager
diff --git a/Engine/Plugins/VirtualProduction/CaptureManager/CaptureManagerCore/Source/CaptureProtocolStack/Private/Control/Messages/ControlResponse.cpp b/Engine/Plugins/VirtualProduction/CaptureManager/CaptureManagerCore/Source/CaptureProtocolStack/Private/Control/Messages/ControlResponse.cpp
--- a/Engine/Plugins/VirtualProduction/CaptureManager/CaptureManagerCore/Source/CaptureProtocolStack/Private/Control/Messages/ControlResponse.cpp
+++ b/Engine/Plugins/VirtualProduction/CaptureManager/CaptureManagerCore/Source/CaptureProtocolStack/Private/Control/Messages/ControlResponse.cpp
@@ -2,8 +2,12 @@
 
 #include "Control/Messages/ControlResponse.h"
 #include "Control/Messages/Constants.h"
+#include "Control/Messages/ProtocolConstants.h"
 
 #include "Control/Messages/ControlJsonUtilities.h"
+#include "Control/Capabilities/Utils.h"
+
+#include "Misc/Paths.h"
 
 namespace UE::CaptureManager
 {
@@ -23,6 +27,11 @@ TProtocolResult<void> FControlResponse::Parse(TSharedPtr<FJsonObject> InBody)
 	return ResultOk;
 }
 
+void FControlResponse::SetAddressPath(const FString& InAddressPath)
+{
+	AddressPath = InAddressPath;
+}
+
 const FString& FControlResponse::GetAddressPath() const
 {
 	return AddressPath;
@@ -364,4 +373,307 @@ FStopStreamingResponse::FStopStreamingResponse()
 {
 }
 
-}
\ No newline at end of file
+FListAllCapabilitiesResponse::FListAllCapabilitiesResponse()
+	: FControlResponse(FPaths::Combine(CPS::AddressPaths::GCapabilitiesRoot, CPS::AddressPaths::GCapabilitiesListAll))
+{
+}
+
+TProtocolResult<void> FListAllCapabilitiesResponse::Parse(TSharedPtr<FJsonObject> InBody)
+{
+	if (!InBody.IsValid())
+	{
+		return FCaptureProtocolError(TEXT("Body not found in capabilities list all response."));
+	}
+
+	const TArray<TSharedPtr<FJsonValue>>* CapabilitiesPtr;
+	CHECK_PARSE(FJsonUtility::ParseArray(InBody, CPS::Capabilities::GValue, CapabilitiesPtr));
+
+	for (const TSharedPtr<FJsonValue>& JsonValue : *CapabilitiesPtr)
+	{
+		TSharedPtr<FJsonObject> JsonObject = JsonValue->AsObject();
+
+		TProtocolResult<FCapability> CapabilityResult = ParseCapability(JsonObject);
+		if (CapabilityResult.IsValid())
+		{
+			Capabilities.Add(MoveTemp(CapabilityResult.GetValue()));
+		}
+		else
+		{
+			return CapabilityResult.GetError();
+		}
+	}
+
+	return ResultOk;
+}
+
+static TProtocolResult<void> CpsAccessStringToEnum(const FString& InAccessString, ECapabilityAccess::Type& OutType)
+{
+	if (InAccessString == CPS::Capabilities::GReadOnly)
+	{
+		OutType = ECapabilityAccess::Type::ReadOnly;
+	}
+	else if (InAccessString == CPS::Capabilities::GReadWrite)
+	{
+		OutType = ECapabilityAccess::Type::ReadWrite;
+	}
+	else
+	{
+		return FCaptureProtocolError(FString::Format(TEXT("Invalid access type value: {0}"), { InAccessString }));
+	}
+	return ResultOk;
+}
+
+TProtocolResult<FCapability> FListAllCapabilitiesResponse::ParseCapability(const TSharedPtr<FJsonObject>& InJsonObject)
+{
+	FCapability Capability;
+
+	CHECK_PARSE(FJsonUtility::ParseString(InJsonObject, CPS::Capabilities::GId, Capability.Id));
+	CHECK_PARSE(FJsonUtility::ParseString(InJsonObject, CPS::Capabilities::GName, Capability.Name));
+
+	const TArray<TSharedPtr<FJsonValue>>* JsonProperties;
+	CHECK_PARSE(FJsonUtility::ParseArray(InJsonObject, CPS::Capabilities::GProperties, JsonProperties));
+
+	for (const TSharedPtr<FJsonValue>& Value : *JsonProperties)
+	{
+		TSharedPtr<FJsonObject> JsonProperty = Value->AsObject();
+		FCapabilityProperty Property;
+
+		CHECK_PARSE(FJsonUtility::ParseString(JsonProperty, CPS::Capabilities::GId, Property.Id));
+		CHECK_PARSE(FJsonUtility::ParseString(JsonProperty, CPS::Capabilities::GName, Property.Name));
+
+		FString PropertyTypeString;
+		CHECK_PARSE(FJsonUtility::ParseString(JsonProperty, CPS::Capabilities::GType, PropertyTypeString));
+		CHECK_PARSE(FCapabilityUtilities::TypeStringToEnum(PropertyTypeString, Property.Type));
+
+		FString PropertyAccessString;
+		CHECK_PARSE(FJsonUtility::ParseString(JsonProperty, CPS::Capabilities::GAccess, PropertyAccessString));
+		CHECK_PARSE(CpsAccessStringToEnum(PropertyAccessString, Property.Access));
+
+		if (JsonProperty->HasField(CPS::Capabilities::GMin))
+		{
+			TVariant<int, float> MinVal;
+
+			float MinLimit;
+			CHECK_PARSE(FJsonUtility::ParseNumber(JsonProperty, CPS::Capabilities::GMin, MinLimit));
+
+			if (Property.Type == ECapabilityValue::Type::Integer)
+			{
+				MinVal.Set<int>(MinLimit);
+			}
+			else
+			{
+				MinVal.Set<float>(MinLimit);
+			}
+			Property.Min.Emplace(MoveTemp(MinVal));
+		}
+
+		if (JsonProperty->HasField(CPS::Capabilities::GMax))
+		{
+			TVariant<int, float> MaxVal;
+
+			float MaxLimit;
+			CHECK_PARSE(FJsonUtility::ParseNumber(JsonProperty, CPS::Capabilities::GMax, MaxLimit));
+
+			if (Property.Type == ECapabilityValue::Type::Integer)
+			{
+				MaxVal.Set<int>(MaxLimit);
+			}
+			else
+			{
+				MaxVal.Set<float>(MaxLimit);
+			}
+			Property.Max.Emplace(MoveTemp(MaxVal));
+		}
+
+		if (Property.Type == ECapabilityValue::Type::Enum)
+		{
+			TArray<FString> EnumOptions;
+			const TArray<TSharedPtr<FJsonValue>>* JsonEnumOptions;
+			CHECK_PARSE(FJsonUtility::ParseArray(JsonProperty, CPS::Capabilities::GEnumOptions, JsonEnumOptions));
+
+			for (const TSharedPtr<FJsonValue>& EnumOption : *JsonEnumOptions)
+			{
+				EnumOptions.Add(EnumOption->AsString());
+			}
+
+			Property.EnumOptions.Emplace(MoveTemp(EnumOptions));
+		}
+
+		CHECK_PARSE(FCapabilityUtilities::ValueFromJsonObject(
+			JsonProperty, CPS::Capabilities::GCurrentValue, Property.Type, Property.CurrentValue));
+
+		Capability.Properties.Add(Property.Id, MoveTemp(Property));
+	}
+
+	const TArray<TSharedPtr<FJsonValue>>* JsonCommands;
+	CHECK_PARSE(FJsonUtility::ParseArray(InJsonObject, CPS::Capabilities::GCommands, JsonCommands));
+
+	for (const TSharedPtr<FJsonValue>& Value : *JsonCommands)
+	{
+		TSharedPtr<FJsonObject> JsonCommand = Value->AsObject();
+		FCapabilityCommand Command;
+
+		CHECK_PARSE(FJsonUtility::ParseString(JsonCommand, CPS::Capabilities::GId, Command.Id));
+		CHECK_PARSE(FJsonUtility::ParseString(JsonCommand, CPS::Capabilities::GName, Command.Name));
+
+		const TArray<TSharedPtr<FJsonValue>>* JsonParameters;
+		CHECK_PARSE(FJsonUtility::ParseArray(JsonCommand, CPS::Capabilities::GParameters, JsonParameters));
+
+		for (const TSharedPtr<FJsonValue>& ParameterValue : *JsonParameters)
+		{
+			TSharedPtr<FJsonObject> JsonParameter = ParameterValue->AsObject();
+			FCommandParameterDescriptor ParameterDescriptor;
+
+			CHECK_PARSE(FJsonUtility::ParseString(JsonParameter, CPS::Capabilities::GName, ParameterDescriptor.Name));
+
+			FString Type;
+			CHECK_PARSE(FJsonUtility::ParseString(JsonParameter, CPS::Capabilities::GType, Type));
+			CHECK_PARSE(FCapabilityUtilities::TypeStringToEnum(Type, ParameterDescriptor.Type));
+
+			CHECK_PARSE(FJsonUtility::ParseBool(JsonParameter, CPS::Capabilities::GOptional, ParameterDescriptor.Optional));
+
+			CHECK_PARSE(FCapabilityUtilities::ValueFromJsonObject(
+				JsonParameter, CPS::Capabilities::GDefaultValue, ParameterDescriptor.Type, ParameterDescriptor.DefaultValue));
+
+			Command.Parameters.Add(ParameterDescriptor);
+		}
+
+		if (JsonCommand->HasField(CPS::Capabilities::GReturnType))
+		{
+			FString ReturnTypeString;
+			CHECK_PARSE(FJsonUtility::ParseString(JsonCommand, CPS::Capabilities::GReturnType, ReturnTypeString));
+			ECapabilityValue::Type ReturnType;
+			CHECK_PARSE(FCapabilityUtilities::TypeStringToEnum(ReturnTypeString, ReturnType));
+			Command.ReturnType = ReturnType;
+		}
+
+		Capability.Commands.Add(Command.Id, MoveTemp(Command));
+	}
+
+	const TArray<TSharedPtr<FJsonValue>>* JsonEvents;
+	CHECK_PARSE(FJsonUtility::ParseArray(InJsonObject, CPS::Capabilities::GEvents, JsonEvents));
+
+	for (const TSharedPtr<FJsonValue>& Value : *JsonEvents)
+	{
+		TSharedPtr<FJsonObject> JsonEvent = Value->AsObject();
+		FCapabilityEvent Event;
+
+		CHECK_PARSE(FJsonUtility::ParseString(JsonEvent, CPS::Capabilities::GId, Event.Id));
+		CHECK_PARSE(FJsonUtility::ParseString(JsonEvent, CPS::Capabilities::GName, Event.Name));
+
+		const TArray<TSharedPtr<FJsonValue>>* JsonArguments;
+		CHECK_PARSE(FJsonUtility::ParseArray(JsonEvent, CPS::Capabilities::GArguments, JsonArguments));
+
+		for (const TSharedPtr<FJsonValue>& ArgumentValue : *JsonArguments)
+		{
+			TSharedPtr<FJsonObject> JsonArgument = ArgumentValue->AsObject();
+
+			FEventArgumentDescriptor CpsArgumentDescriptor;
+			CHECK_PARSE(FJsonUtility::ParseString(JsonArgument, CPS::Capabilities::GName, CpsArgumentDescriptor.Name));
+			FString ArgumentTypeString;
+			CHECK_PARSE(FJsonUtility::ParseString(JsonArgument, CPS::Capabilities::GType, ArgumentTypeString));
+			CHECK_PARSE(FCapabilityUtilities::TypeStringToEnum(ArgumentTypeString, CpsArgumentDescriptor.Type));
+
+			Event.Arguments.Add(MoveTemp(CpsArgumentDescriptor));
+		}
+
+		Capability.Events.Add(Event.Id, MoveTemp(Event));
+	}
+
+	return Capability;
+}
+
+TArray<FCapability> FListAllCapabilitiesResponse::GetCapabilities()
+{
+	return Capabilities;
+}
+
+FSetCapabilityPropertyResponse::FSetCapabilityPropertyResponse()
+	: FControlResponse(CPS::AddressPaths::GCapabilitiesRoot)
+{
+}
+
+FGetCapabilityPropertyResponse::FGetCapabilityPropertyResponse()
+	: FControlResponse(CPS::AddressPaths::GCapabilitiesRoot)
+{
+}
+
+TProtocolResult<void> FGetCapabilityPropertyResponse::Parse(TSharedPtr<FJsonObject> InBody)
+{
+	if (!InBody.IsValid())
+	{
+		return FCaptureProtocolError(TEXT("Body not found in capability property operation response."));
+	}
+
+	Body = MoveTemp(InBody);
+
+	return ResultOk;
+}
+
+TProtocolResult<FCapabilityValue> FGetCapabilityPropertyResponse::GetValue(ECapabilityValue::Type InExpectedType) const
+{
+	if (Body->HasField(CPS::Capabilities::GValue))
+	{
+		FCapabilityValue CapabilityValue;
+		TProtocolResult<void> Result =
+			FCapabilityUtilities::ValueFromJsonObject(Body, CPS::Capabilities::GValue, InExpectedType, CapabilityValue);
+
+		if (Result.HasError())
+		{
+			return Result.GetError();
+		}
+
+		return CapabilityValue;
+	}
+
+	return FCaptureProtocolError(TEXT("Command execution did not return value."));
+}
+
+FExecuteCapabilityCommandResponse::FExecuteCapabilityCommandResponse()
+	: FControlResponse(CPS::AddressPaths::GCapabilitiesRoot)
+{
+}
+
+TProtocolResult<void> FExecuteCapabilityCommandResponse::Parse(TSharedPtr<FJsonObject> InBody)
+{
+	if (!InBody.IsValid())
+	{
+		return FCaptureProtocolError(TEXT("Body not found in capability execute command operation response."));
+	}
+
+	Body = MoveTemp(InBody);
+	
+	return ResultOk;
+}
+
+TProtocolResult<FCapabilityValue> FExecuteCapabilityCommandResponse::GetReturnValue(ECapabilityValue::Type InExpectedType) const
+{
+	if (!Body->HasField(CPS::Capabilities::GValue))
+	{
+		return FCaptureProtocolError(TEXT("Command execution did not return value."));
+	}
+
+	FCapabilityValue CapabilityValue;
+
+	TProtocolResult<void> Result = FCapabilityUtilities::ValueFromJsonObject(Body, CPS::Capabilities::GValue, InExpectedType, CapabilityValue);
+
+	if (Result.HasError())
+	{
+		FString Message = FString::Format(TEXT("Failed to convert command return value to an expected type: {0}"), { *Result.GetError().GetMessage() });
+		return FCaptureProtocolError(MoveTemp(Message));
+	}
+
+	return CapabilityValue;
+}
+
+FCapabilitySubscribeResponse::FCapabilitySubscribeResponse()
+	: FControlResponse(CPS::AddressPaths::GCapabilitiesRoot)
+{
+}
+
+FCapabilityUnsubscribeResponse::FCapabilityUnsubscribeResponse()
+	: FControlResponse(CPS::AddressPaths::GCapabilitiesRoot)
+{
+}
+
+} // namespace UE::CaptureManager
diff --git a/Engine/Plugins/VirtualProduction/CaptureManager/CaptureManagerCore/Source/CaptureProtocolStack/Private/Control/Messages/ControlUpdate.cpp b/Engine/Plugins/VirtualProduction/CaptureManager/CaptureManagerCore/Source/CaptureProtocolStack/Private/Control/Messages/ControlUpdate.cpp
--- a/Engine/Plugins/VirtualProduction/CaptureManager/CaptureManagerCore/Source/CaptureProtocolStack/Private/Control/Messages/ControlUpdate.cpp
+++ b/Engine/Plugins/VirtualProduction/CaptureManager/CaptureManagerCore/Source/CaptureProtocolStack/Private/Control/Messages/ControlUpdate.cpp
@@ -2,8 +2,10 @@
 
 #include "Control/Messages/ControlUpdate.h"
 #include "Control/Messages/Constants.h"
+#include "Control/Messages/ProtocolConstants.h"
 
 #include "Control/Messages/ControlJsonUtilities.h"
+#include "Control/Capabilities/Utils.h"
 
 namespace UE::CaptureManager
 {
@@ -42,6 +44,10 @@ TProtocolResult<TSharedRef<FControlUpdate>> FControlUpdateCreator::Create(const
 	{
 		return TProtocolResult<TSharedRef<FControlUpdate>>(MakeShared<FThermalStateUpdate>());
 	}
+	else if (InAddressPath.StartsWith(CPS::AddressPaths::GCapabilitiesRoot))
+	{
+		return TProtocolResult<TSharedRef<FControlUpdate>>(MakeShared<FEventUpdate>(InAddressPath));
+	}
 	else
 	{
 		return FCaptureProtocolError(TEXT("Unknown update arrived"));
@@ -205,4 +211,58 @@ FThermalStateUpdate::EState FThermalStateUpdate::ConvertState(const FString& InS
 	}
 }
 
-}
\ No newline at end of file
+FEventUpdate::FEventUpdate(FString InAddressPath)
+	: FControlUpdate(MoveTemp(InAddressPath))
+{
+}
+
+TProtocolResult<void> FEventUpdate::Parse(TSharedPtr<FJsonObject> InBody)
+{
+	if (!InBody.IsValid())
+	{
+		return FCaptureProtocolError(TEXT("Body not found in capability event update."));
+	}
+
+	const TArray<TSharedPtr<FJsonValue>>* ArgumentsPtr;
+	CHECK_PARSE(FJsonUtility::ParseArray(InBody, CPS::Capabilities::GArguments, ArgumentsPtr));
+
+	for (TSharedPtr<FJsonValue> JsonValue : *ArgumentsPtr)
+	{
+		TSharedPtr<FJsonObject> JsonObject = JsonValue->AsObject();
+
+		FString ArgumentId;
+		CHECK_PARSE(FJsonUtility::ParseString(JsonObject, CPS::Capabilities::GId, ArgumentId));		
+
+		Arguments.Add(MoveTemp(ArgumentId), MoveTemp(JsonObject));
+	}
+
+	return ResultOk;
+}
+
+const TProtocolResult<TMap<FString, FCapabilityValue>> FEventUpdate::GetArguments(const TArray<FEventArgumentDescriptor>& InArgumentsDescriptor) const
+{
+	TMap<FString, FCapabilityValue> ArgumentValues;
+
+	TArray<FString> ArgumentNames;
+	
+	for (const FEventArgumentDescriptor& ArgumentDescriptor : InArgumentsDescriptor)
+	{
+		const TSharedPtr<FJsonObject>* JsonObjectPtr = Arguments.Find(ArgumentDescriptor.Name);
+		const ECapabilityValue::Type ArgumentType = ArgumentDescriptor.Type;
+		if (JsonObjectPtr)
+		{
+			FCapabilityValue CapabilityValue;
+			TProtocolResult<void> Result = FCapabilityUtilities::ValueFromJsonObject(*JsonObjectPtr, CPS::Capabilities::GValue, ArgumentType, CapabilityValue);
+			if (Result.HasError())
+			{
+				return Result.GetError();
+			}
+			ArgumentValues.Add({ ArgumentDescriptor.Name, CapabilityValue });
+		}
+		
+	}
+	
+	return ArgumentValues;
+}
+
+} // namespace UE::CaptureManager
\ No newline at end of file
diff --git a/Engine/Plugins/VirtualProduction/CaptureManager/CaptureManagerCore/Source/CaptureProtocolStack/Private/Control/Messages/ProtocolConstants.h b/Engine/Plugins/VirtualProduction/CaptureManager/CaptureManagerCore/Source/CaptureProtocolStack/Private/Control/Messages/ProtocolConstants.h
new file mode 100644
--- /dev/null
+++ b/Engine/Plugins/VirtualProduction/CaptureManager/CaptureManagerCore/Source/CaptureProtocolStack/Private/Control/Messages/ProtocolConstants.h
@@ -0,0 +1,53 @@
+// Copyright Epic Games, Inc. All Rights Reserved.
+
+#pragma once
+
+#include "Containers/UnrealString.h"
+
+namespace UE::CaptureManager::CPS::AddressPaths
+{
+// Capabilities
+extern CAPTUREPROTOCOLSTACK_API const TCHAR* GCapabilitiesRoot;
+extern CAPTUREPROTOCOLSTACK_API const TCHAR* GCapabilitiesListAll;
+extern CAPTUREPROTOCOLSTACK_API const TCHAR* GCapabilityProperties;
+extern CAPTUREPROTOCOLSTACK_API const TCHAR* GCapabilityPropertySet;
+extern CAPTUREPROTOCOLSTACK_API const TCHAR* GCapabilityPropertyGet;
+extern CAPTUREPROTOCOLSTACK_API const TCHAR* GCapabilityCommands;
+extern CAPTUREPROTOCOLSTACK_API const TCHAR* GCapabilityCommandExecute;
+extern CAPTUREPROTOCOLSTACK_API const TCHAR* GCapabilityEvents;
+extern CAPTUREPROTOCOLSTACK_API const TCHAR* GCapabilityEventSubscribe;
+extern CAPTUREPROTOCOLSTACK_API const TCHAR* GCapabilityEventUnsubscribe;
+} // namespace UE::CaptureManager::CPS::AddressPaths
+
+namespace UE::CaptureManager::CPS::Capabilities
+{
+extern CAPTUREPROTOCOLSTACK_API const TCHAR* GId;
+extern CAPTUREPROTOCOLSTACK_API const TCHAR* GName;
+extern CAPTUREPROTOCOLSTACK_API const TCHAR* GValue;
+extern CAPTUREPROTOCOLSTACK_API const TCHAR* GArguments;
+extern CAPTUREPROTOCOLSTACK_API const TCHAR* GParameters;
+extern CAPTUREPROTOCOLSTACK_API const TCHAR* GProperties;
+extern CAPTUREPROTOCOLSTACK_API const TCHAR* GType;
+extern CAPTUREPROTOCOLSTACK_API const TCHAR* GAccess;
+extern CAPTUREPROTOCOLSTACK_API const TCHAR* GReadOnly;
+extern CAPTUREPROTOCOLSTACK_API const TCHAR* GReadWrite;
+extern CAPTUREPROTOCOLSTACK_API const TCHAR* GMin;
+extern CAPTUREPROTOCOLSTACK_API const TCHAR* GMax;
+extern CAPTUREPROTOCOLSTACK_API const TCHAR* GEnumOptions;
+extern CAPTUREPROTOCOLSTACK_API const TCHAR* GCurrentValue;
+extern CAPTUREPROTOCOLSTACK_API const TCHAR* GGet;
+extern CAPTUREPROTOCOLSTACK_API const TCHAR* GSet;
+extern CAPTUREPROTOCOLSTACK_API const TCHAR* GCommands;
+extern CAPTUREPROTOCOLSTACK_API const TCHAR* GExecute;
+extern CAPTUREPROTOCOLSTACK_API const TCHAR* GOptional;
+extern CAPTUREPROTOCOLSTACK_API const TCHAR* GReturnType;
+extern CAPTUREPROTOCOLSTACK_API const TCHAR* GDefaultValue;
+extern CAPTUREPROTOCOLSTACK_API const TCHAR* GEvents;
+extern CAPTUREPROTOCOLSTACK_API const TCHAR* GSubscribe;
+extern CAPTUREPROTOCOLSTACK_API const TCHAR* GUnsubscribe;
+extern CAPTUREPROTOCOLSTACK_API const TCHAR* GBoolean;
+extern CAPTUREPROTOCOLSTACK_API const TCHAR* GInteger;
+extern CAPTUREPROTOCOLSTACK_API const TCHAR* GFloat;
+extern CAPTUREPROTOCOLSTACK_API const TCHAR* GString;
+extern CAPTUREPROTOCOLSTACK_API const TCHAR* GEnumeration;
+} // namespace UE::CaptureManager::CPS::Capabilities
diff --git a/Engine/Plugins/VirtualProduction/CaptureManager/CaptureManagerCore/Source/CaptureProtocolStack/Public/Control/Capabilities/Capability.h b/Engine/Plugins/VirtualProduction/CaptureManager/CaptureManagerCore/Source/CaptureProtocolStack/Public/Control/Capabilities/Capability.h
new file mode 100644
--- /dev/null
+++ b/Engine/Plugins/VirtualProduction/CaptureManager/CaptureManagerCore/Source/CaptureProtocolStack/Public/Control/Capabilities/Capability.h
@@ -0,0 +1,90 @@
+// Copyright Epic Games, Inc. All Rights Reserved.
+
+#pragma once
+
+#include "Containers/Map.h"
+#include "Misc/Optional.h"
+#include "Misc/TVariant.h"
+
+class FString;
+
+namespace UE::CaptureManager
+{
+namespace ECapabilityValue
+{
+enum Type
+{
+	Boolean,
+	Integer,
+	Float,
+	String,
+	Enum
+};
+} // namespace ECapabilityValue
+
+namespace ECapabilityAccess
+{
+enum Type
+{
+	ReadOnly = 0,
+	ReadWrite = 1
+};
+} // namespace ECapabilityAccess
+
+using FCapabilityValue = TVariant<bool, int, float, FString>;
+
+using FCapabilityValues = TMap<FString, FCapabilityValue>;
+
+struct FCapabilityProperty
+{
+	FString Id;
+	FString Name;
+	ECapabilityValue::Type Type;
+	ECapabilityAccess::Type Access;
+	TOptional<TVariant<int, float>> Min;
+	TOptional<TVariant<int, float>> Max;
+	TOptional<TArray<FString>> EnumOptions;
+	FCapabilityValue CurrentValue;
+};
+
+struct FCommandParameterDescriptor
+{
+	FString Name;
+	ECapabilityValue::Type Type;
+	bool Optional;
+	FCapabilityValue DefaultValue;
+};
+
+struct FCapabilityCommand
+{
+	FString Id;
+	FString Name;
+
+	TArray<FCommandParameterDescriptor> Parameters;
+	TOptional<ECapabilityValue::Type> ReturnType;
+};
+
+struct FEventArgumentDescriptor
+{
+	FString Name;
+	ECapabilityValue::Type Type;
+};
+
+struct FCapabilityEvent
+{
+	FString Id;
+	FString Name;
+
+	TArray<FEventArgumentDescriptor> Arguments;
+};
+
+struct FCapability
+{
+	FString Id;
+	FString Name;
+	TMap<FString, FCapabilityProperty> Properties;
+	TMap<FString, FCapabilityCommand> Commands;
+	TMap<FString, FCapabilityEvent> Events;
+};
+
+} // namespace UE::CaptureManager
diff --git a/Engine/Plugins/VirtualProduction/CaptureManager/CaptureManagerCore/Source/CaptureProtocolStack/Public/Control/ControlMessenger.h b/Engine/Plugins/VirtualProduction/CaptureManager/CaptureManagerCore/Source/CaptureProtocolStack/Public/Control/ControlMessenger.h
--- a/Engine/Plugins/VirtualProduction/CaptureManager/CaptureManagerCore/Source/CaptureProtocolStack/Public/Control/ControlMessenger.h
+++ b/Engine/Plugins/VirtualProduction/CaptureManager/CaptureManagerCore/Source/CaptureProtocolStack/Public/Control/ControlMessenger.h
@@ -127,10 +127,12 @@ public:
 
 		if (!ResponseMessage.GetErrorName().IsEmpty())
 		{
-			return FCaptureProtocolError(TEXT("Server responded with error: ") + ResponseMessage.GetErrorName());
+			return FCaptureProtocolError(FString::Format(TEXT("Server responded with error: {0} : {1}"), { ResponseMessage.GetErrorName(), ResponseMessage.GetErrorDescription() }));
 		}
 
 		typename RequestType::ResponseType Response;
+		Response.SetAddressPath(InRequest.GetAddressPath());
+
 		RequestContexts.Remove(TransactionId);
 		
 		TProtocolResult<void> ParseResult = Response.Parse(ResponseMessage.GetBody());
diff --git a/Engine/Plugins/VirtualProduction/CaptureManager/CaptureManagerCore/Source/CaptureProtocolStack/Public/Control/Messages/ControlRequest.h b/Engine/Plugins/VirtualProduction/CaptureManager/CaptureManagerCore/Source/CaptureProtocolStack/Public/Control/Messages/ControlRequest.h
--- a/Engine/Plugins/VirtualProduction/CaptureManager/CaptureManagerCore/Source/CaptureProtocolStack/Public/Control/Messages/ControlRequest.h
+++ b/Engine/Plugins/VirtualProduction/CaptureManager/CaptureManagerCore/Source/CaptureProtocolStack/Public/Control/Messages/ControlRequest.h
@@ -96,10 +96,10 @@ public:
 	using ResponseType = FStartRecordingTakeResponse;
 
 	UE_API FStartRecordingTakeRequest(FString InSlateName,
-							   uint16 InTakeNumber,
-							   TOptional<FString> InSubject = TOptional<FString>(),
-							   TOptional<FString> InScenario = TOptional<FString>(),
-							   TOptional<TArray<FString>> InTags = TOptional<TArray<FString>>());
+									  uint16 InTakeNumber,
+									  TOptional<FString> InSubject = TOptional<FString>(),
+									  TOptional<FString> InScenario = TOptional<FString>(),
+									  TOptional<TArray<FString>> InTags = TOptional<TArray<FString>>());
 
 	UE_API virtual TSharedPtr<FJsonObject> GetBody() const override;
 
@@ -211,6 +211,64 @@ private:
 	TOptional<TArray<FString>> SubjectIds;
 };
 
-}
+class FListAllCapabilitiesRequest final : public FControlRequest
+{
+public:
+	using ResponseType = FListAllCapabilitiesResponse;
+
+	UE_API FListAllCapabilitiesRequest();
+};
+
+class FSetCapabilityPropertyRequest final : public FControlRequest
+{
+public:
+	using ResponseType = FSetCapabilityPropertyResponse;
+
+	UE_API FSetCapabilityPropertyRequest(FString InCapabilityId, FString InPropertyId, FCapabilityValue InValue);
+
+	UE_API virtual TSharedPtr<FJsonObject> GetBody() const override;
+
+private:
+	FCapabilityValue Value;
+};
+
+class FGetCapabilityPropertyRequest final : public FControlRequest
+{
+public:
+	using ResponseType = FGetCapabilityPropertyResponse;
+
+	UE_API FGetCapabilityPropertyRequest(FString InCapabilityId, FString InPropertyId);
+};
+
+class FExecuteCapabilityCommandRequest final : public FControlRequest
+{
+public:
+	using ResponseType = FExecuteCapabilityCommandResponse;
+
+	UE_API FExecuteCapabilityCommandRequest(FString InCapabilityId, FString InCommandId, FCapabilityValues InArguments);
+
+	UE_API virtual TSharedPtr<FJsonObject> GetBody() const override;
+
+private:
+	FCapabilityValues Arguments;
+};
+
+class FCapabilitySubscribeRequest final : public FControlRequest
+{
+public:
+	using ResponseType = FCapabilitySubscribeResponse;
+
+	UE_API FCapabilitySubscribeRequest(FString InCapabilityId);
+};
+
+class FCapabilityUnsubscribeRequest final : public FControlRequest
+{
+public:
+	using ResponseType = FCapabilityUnsubscribeResponse;
+
+	UE_API FCapabilityUnsubscribeRequest(FString InCapabilityId);
+};
+
+} // namespace UE::CaptureManager
 
 #undef UE_API
diff --git a/Engine/Plugins/VirtualProduction/CaptureManager/CaptureManagerCore/Source/CaptureProtocolStack/Public/Control/Messages/ControlResponse.h b/Engine/Plugins/VirtualProduction/CaptureManager/CaptureManagerCore/Source/CaptureProtocolStack/Public/Control/Messages/ControlResponse.h
--- a/Engine/Plugins/VirtualProduction/CaptureManager/CaptureManagerCore/Source/CaptureProtocolStack/Public/Control/Messages/ControlResponse.h
+++ b/Engine/Plugins/VirtualProduction/CaptureManager/CaptureManagerCore/Source/CaptureProtocolStack/Public/Control/Messages/ControlResponse.h
@@ -3,6 +3,7 @@
 #pragma once
 
 #include "ControlJsonUtilities.h"
+#include "Control/Capabilities/Capability.h"
 
 #define UE_API CAPTUREPROTOCOLSTACK_API
 
@@ -18,6 +19,7 @@ public:
 
 	UE_API virtual TProtocolResult<void> Parse(TSharedPtr<FJsonObject> InBody);
 
+	UE_API void SetAddressPath(const FString& InAddressPath);
 	UE_API const FString& GetAddressPath() const;
 
 private:
@@ -258,10 +260,66 @@ public:
 class FStopStreamingResponse final : public FControlResponse
 {
 public:
-
 	UE_API FStopStreamingResponse();
 };
 
-}
+class FListAllCapabilitiesResponse final : public FControlResponse
+{
+public:
+	UE_API FListAllCapabilitiesResponse();
+
+	UE_API TProtocolResult<void> Parse(TSharedPtr<FJsonObject> InBody) override;
+
+	UE_API TArray<FCapability> GetCapabilities();
+
+private:
+	TArray<FCapability> Capabilities;
+
+	TProtocolResult<FCapability> ParseCapability(const TSharedPtr<FJsonObject>& InJsonObject);
+};
+
+class FSetCapabilityPropertyResponse final : public FControlResponse
+{
+public:
+	UE_API FSetCapabilityPropertyResponse();
+};
+
+class FGetCapabilityPropertyResponse final : public FControlResponse
+{
+public:
+	UE_API FGetCapabilityPropertyResponse();
+
+	UE_API TProtocolResult<void> Parse(TSharedPtr<FJsonObject> InBody) override;
+
+	UE_API TProtocolResult<FCapabilityValue> GetValue(ECapabilityValue::Type InExpectedType) const;
+
+private:
+	TSharedPtr<FJsonObject> Body;
+};
+
+class FExecuteCapabilityCommandResponse final : public FControlResponse
+{
+public:
+	UE_API FExecuteCapabilityCommandResponse();
+
+	UE_API TProtocolResult<void> Parse(TSharedPtr<FJsonObject> InBody) override;
+
+	UE_API TProtocolResult<FCapabilityValue> GetReturnValue(ECapabilityValue::Type InExpectedType) const;
+
+private:
+	TSharedPtr<FJsonObject> Body;
+};
+
+class FCapabilitySubscribeResponse final : public FControlResponse
+{
+public:
+	UE_API FCapabilitySubscribeResponse();
+};
 
+class FCapabilityUnsubscribeResponse final : public FControlResponse
+{
+public:
+	UE_API FCapabilityUnsubscribeResponse();
+};
+} // namespace UE::CaptureManager
 #undef UE_API
diff --git a/Engine/Plugins/VirtualProduction/CaptureManager/CaptureManagerCore/Source/CaptureProtocolStack/Public/Control/Messages/ControlUpdate.h b/Engine/Plugins/VirtualProduction/CaptureManager/CaptureManagerCore/Source/CaptureProtocolStack/Public/Control/Messages/ControlUpdate.h
--- a/Engine/Plugins/VirtualProduction/CaptureManager/CaptureManagerCore/Source/CaptureProtocolStack/Public/Control/Messages/ControlUpdate.h
+++ b/Engine/Plugins/VirtualProduction/CaptureManager/CaptureManagerCore/Source/CaptureProtocolStack/Public/Control/Messages/ControlUpdate.h
@@ -3,10 +3,12 @@
 #pragma once
 
 #include "ControlJsonUtilities.h"
+#include "Containers/Map.h"
 
 #include "Utility/Error.h"
 
 #include "Delegates/Delegate.h"
+#include "Control/Capabilities/Capability.h"
 
 #define UE_API CAPTUREPROTOCOLSTACK_API
 
@@ -157,6 +159,18 @@ private:
 	EState State;
 };
 
-}
+class FEventUpdate final : public FControlUpdate
+{
+public:
+	UE_API FEventUpdate(FString InAddressPath);
+
+	UE_API virtual UE::CaptureManager::TProtocolResult<void> Parse(TSharedPtr<FJsonObject> InBody) override;
+
+	UE_API const TProtocolResult<TMap<FString, FCapabilityValue>> GetArguments(const TArray<FEventArgumentDescriptor>& InArgumentsDescriptor) const;
+
+private:
+	TMap<FString, TSharedPtr<FJsonObject>> Arguments;
+};
+} // namespace UE::CaptureManager
 
 #undef UE_API
