Fixing rapidjson issue when validating JSON using schema

#Synced-CL 50477824

diff --git a/Engine/Plugins/VirtualProduction/CaptureManager/CaptureManagerCore/Source/CaptureManagerTakeMetadata/Private/CaptureManagerTakeMetadata.cpp b/Engine/Plugins/VirtualProduction/CaptureManager/CaptureManagerCore/Source/CaptureManagerTakeMetadata/Private/CaptureManagerTakeMetadata.cpp
--- a/Engine/Plugins/VirtualProduction/CaptureManager/CaptureManagerCore/Source/CaptureManagerTakeMetadata/Private/CaptureManagerTakeMetadata.cpp
+++ b/Engine/Plugins/VirtualProduction/CaptureManager/CaptureManagerCore/Source/CaptureManagerTakeMetadata/Private/CaptureManagerTakeMetadata.cpp
@@ -688,6 +688,8 @@ struct FTakeMetadataParser::FImpl
 {
 	struct FRapidjsonValidatorBundle
 	{
+		using FPtr = TUniquePtr<FRapidjsonValidatorBundle>;
+
 		rapidjson::Document Document;
 		rapidjson::SchemaDocument SchemaDocument;
 		rapidjson::SchemaValidator Validator;
@@ -695,8 +697,10 @@ struct FTakeMetadataParser::FImpl
 		FRapidjsonValidatorBundle(rapidjson::Document&& InDocument) : Document(MoveTemp(InDocument)), SchemaDocument(Document), Validator(SchemaDocument)
 		{}
 
-		FRapidjsonValidatorBundle(FRapidjsonValidatorBundle&& InOther) : Document(MoveTemp(InOther.Document)), SchemaDocument(Document), Validator(SchemaDocument)
-		{}
+		FRapidjsonValidatorBundle(FRapidjsonValidatorBundle&& InOther) = delete;
+		FRapidjsonValidatorBundle(const FRapidjsonValidatorBundle& InOther) = delete;
+		FRapidjsonValidatorBundle& operator=(FRapidjsonValidatorBundle&& InOther) = delete;
+		FRapidjsonValidatorBundle& operator=(const FRapidjsonValidatorBundle& InOther) = delete;
 	};
 
 	struct FValidator
@@ -705,11 +709,11 @@ struct FTakeMetadataParser::FImpl
 		rapidjson::SchemaValidator& Validator;
 	};
 
-	TOptional<FRapidjsonValidatorBundle> VersionValidatorBundle;
+	TOptional<FRapidjsonValidatorBundle::FPtr> VersionValidatorBundle;
 
-	TMap<FTakeMetadata::FSchemaVersion, FRapidjsonValidatorBundle> ValidatorBundles;
+	TMap<FTakeMetadata::FSchemaVersion, FRapidjsonValidatorBundle::FPtr> ValidatorBundles;
 
-	TValueOrError<FRapidjsonValidatorBundle, FTakeMetadataParserError> CreateValidator(const FString& InSchemaFileName)
+	TValueOrError<FRapidjsonValidatorBundle::FPtr, FTakeMetadataParserError> CreateValidator(const FString& InSchemaFileName)
 	{
 		FString SchemaString;
 
@@ -730,7 +734,7 @@ struct FTakeMetadataParser::FImpl
 			return MakeError(MoveTemp(Error));
 		}
 
-		return MakeValue(FRapidjsonValidatorBundle(MoveTemp(Document)));
+		return MakeValue(MakeUnique<FRapidjsonValidatorBundle>(MoveTemp(Document)));
 	}
 
 
@@ -740,7 +744,7 @@ struct FTakeMetadataParser::FImpl
 		{
 			const FString VersionSchemaFilePath = InSchemasDir / TEXT("version.json");
 
-			TValueOrError<FRapidjsonValidatorBundle, FTakeMetadataParserError> VersionValidatorResult = CreateValidator(VersionSchemaFilePath);
+			TValueOrError<FRapidjsonValidatorBundle::FPtr, FTakeMetadataParserError> VersionValidatorResult = CreateValidator(VersionSchemaFilePath);
 
 			if (VersionValidatorResult.HasError())
 			{
@@ -750,10 +754,10 @@ struct FTakeMetadataParser::FImpl
 			VersionValidatorBundle = VersionValidatorResult.StealValue();
 		}
 
-		VersionValidatorBundle->Validator.Reset();
-		VersionValidatorBundle->Validator.ResetError();
+		VersionValidatorBundle.GetValue()->Validator.Reset();
+		VersionValidatorBundle.GetValue()->Validator.ResetError();
 
-		TOptional<FTakeMetadataValidationError> ValidationError = Validate(InDocument, VersionValidatorBundle->Validator);
+		TOptional<FTakeMetadataValidationError> ValidationError = Validate(InDocument, VersionValidatorBundle.GetValue()->Validator);
 		if (ValidationError.IsSet())
 		{
 			FTakeMetadataParserError Error;
@@ -778,7 +782,7 @@ struct FTakeMetadataParser::FImpl
 		}
 
 		FTakeMetadata::FSchemaVersion Version = VersionResult.StealValue();
-		FRapidjsonValidatorBundle* ValidatorBundle = ValidatorBundles.Find(Version);
+		FRapidjsonValidatorBundle::FPtr* ValidatorBundle = ValidatorBundles.Find(Version);
 
 		FString SchemaFileName = "v";
 		SchemaFileName += FString::FromInt(Version.Major) + ".";
@@ -788,7 +792,7 @@ struct FTakeMetadataParser::FImpl
 		{
 			FString SchemaFilePath = SchemasDir / SchemaFileName;
 
-			TValueOrError<FRapidjsonValidatorBundle, FTakeMetadataParserError> ValidatorResult = CreateValidator(SchemaFilePath);
+			TValueOrError<FRapidjsonValidatorBundle::FPtr, FTakeMetadataParserError> ValidatorResult = CreateValidator(SchemaFilePath);
 
 			if (ValidatorResult.HasError())
 			{
@@ -800,10 +804,10 @@ struct FTakeMetadataParser::FImpl
 			ValidatorBundle = ValidatorBundles.Find(Version);
 		}
 
-		ValidatorBundle->Validator.Reset();
-		ValidatorBundle->Validator.ResetError();
+		(*ValidatorBundle)->Validator.Reset();
+		(*ValidatorBundle)->Validator.ResetError();
 
-		return MakeValue(FValidator{ MoveTemp(Version), ValidatorBundle->Validator });
+		return MakeValue(FValidator{ MoveTemp(Version), (*ValidatorBundle)->Validator });
 	}
 };
 
@@ -884,5 +888,3 @@ FString FTakeMetadataPathUtils::PathTypeToString(FTakeMetadata::FVideo::EPathTyp
 }
 
 #undef LOCTEXT_NAMESPACE
-
-
