Removing unnecessary threads from ControlMessenger

#Synced-CL 51089948

diff --git a/Engine/Plugins/VirtualProduction/CaptureManager/CaptureManagerCore/Source/CaptureProtocolStack/Private/Control/Communication/ControlCommunication.cpp b/Engine/Plugins/VirtualProduction/CaptureManager/CaptureManagerCore/Source/CaptureProtocolStack/Private/Control/Communication/ControlCommunication.cpp
--- a/Engine/Plugins/VirtualProduction/CaptureManager/CaptureManagerCore/Source/CaptureProtocolStack/Private/Control/Communication/ControlCommunication.cpp
+++ b/Engine/Plugins/VirtualProduction/CaptureManager/CaptureManagerCore/Source/CaptureProtocolStack/Private/Control/Communication/ControlCommunication.cpp
@@ -11,8 +11,6 @@ DEFINE_LOG_CATEGORY(LogCPSControlCommunication)
 
 FControlCommunication::FControlCommunication()
 	: ReceiveHandler(nullptr)
-	, SynchronizedReceiver(FRunnerType::FOnProcess::CreateRaw(this, &FControlCommunication::OnProcessReceivedPacket))
-	, SynchronizedSender(FRunnerType::FOnProcess::CreateRaw(this, &FControlCommunication::OnProcessSentPacket))
 	, SynchronizedRunnable(TQueueRunner<TSharedPtr<FRunnable>>::FOnProcess::CreateRaw(this, &FControlCommunication::RunnableHandler))
 {
 }
@@ -60,7 +58,17 @@ bool FControlCommunication::IsRunning() const
 
 void FControlCommunication::SendMessage(FControlPacket InMessage)
 {
-	SynchronizedSender.Add(MoveTemp(InMessage));
+	static constexpr int32 MaxNumberOfMessagesInQueue = 20;
+
+	if (SenderQueueCurrentSize < MaxNumberOfMessagesInQueue)
+	{
+		SenderQueue.Enqueue(MoveTemp(InMessage));
+		++SenderQueueCurrentSize;
+	}
+	else
+	{
+		UE_LOG(LogCPSControlCommunication, Error, TEXT("Message is not scheduled for sending due to busy senders queue"));
+	}
 }
 
 void FControlCommunication::SetReceiveHandler(FOnPacketReceived InReceiveHandler)
@@ -73,16 +81,28 @@ void FControlCommunication::SetCommunicationStoppedHandler(FCommunicationStopped
 	CommunicationStoppedHandler = MoveTemp(InCommunicationStoppedHandler);
 }
 
-void FControlCommunication::OnProcessReceivedPacket(FControlPacket InMessage)
+void FControlCommunication::RunnableHandler(TSharedPtr<FRunnable> InRunnable)
 {
-	ReceiveHandler.ExecuteIfBound(MoveTemp(InMessage));
+	InRunnable->Run();
+}
+
+void FControlCommunication::SendNextPacket()
+{
+	FControlPacket Packet;
+
+	if (SenderQueue.Dequeue(Packet))
+	{
+		--SenderQueueCurrentSize;
+
+		SendPacket(MoveTemp(Packet));
+	}
 }
 
-void FControlCommunication::OnProcessSentPacket(FControlPacket InMessage)
+void FControlCommunication::SendPacket(FControlPacket InPacket)
 {
 	FTcpClientWriter Writer(Client);
 
-	TProtocolResult<void> SerializeResult = FControlPacket::Serialize(InMessage, Writer);
+	TProtocolResult<void> SerializeResult = FControlPacket::Serialize(InPacket, Writer);
 	if (SerializeResult.HasError())
 	{
 		FCaptureProtocolError Error = SerializeResult.StealError();
@@ -91,16 +111,21 @@ void FControlCommunication::OnProcessSentPacket(FControlPacket InMessage)
 	}
 }
 
-void FControlCommunication::RunnableHandler(TSharedPtr<FRunnable> InRunnable)
+void FControlCommunication::ReceiveNextPacket()
 {
-	InRunnable->Run();
+	FControlPacket Packet;
+	if (ReceiverQueue.Dequeue(Packet))
+	{
+		ReceiveHandler.ExecuteIfBound(MoveTemp(Packet));
+	}
 }
 
 TProtocolResult<FControlPacketHeader> FControlCommunication::ReceiveControlHeader()
 {
 	FTcpClientReader Reader(Client);
 
-	return FControlPacketHeader::Deserialize(Reader);
+	static constexpr uint32 HeaderTimeout = 50; // 50 milliseconds as we are doing reading and writing to the socket in a single thread
+	return FControlPacketHeader::Deserialize(Reader, HeaderTimeout);
 }
 
 TProtocolResult<void> FControlCommunication::ReceiveControlPacket(const FControlPacketHeader& InHeader)
@@ -113,7 +138,7 @@ TProtocolResult<void> FControlCommunication::ReceiveControlPacket(const FControl
 		return PacketDeserialize.StealError();
 	}
 
-	SynchronizedReceiver.Add(PacketDeserialize.StealValue());
+	ReceiverQueue.Enqueue(PacketDeserialize.StealValue());
 
 	return ResultOk;
 }
@@ -130,6 +155,18 @@ uint32 FControlCommunication::FCommunicationRunnable::Run()
 {
 	while (bIsRunning)
 	{
+		static constexpr int32 MaxMessagesPerIteration = 10;
+
+		for (int32 Counter = 0; Counter < MaxMessagesPerIteration && !Communication.SenderQueue.IsEmpty(); ++Counter)
+		{
+			Communication.SendNextPacket();
+		}
+
+		for (int32 Counter = 0; Counter < MaxMessagesPerIteration && !Communication.ReceiverQueue.IsEmpty(); ++Counter)
+		{
+			Communication.ReceiveNextPacket();
+		}
+
 		TProtocolResult<FControlPacketHeader> ReceiveHeaderResult = Communication.ReceiveControlHeader();
 		if (ReceiveHeaderResult.HasError())
 		{
@@ -171,6 +208,7 @@ void FControlCommunication::FCommunicationRunnable::HandleError(const FCapturePr
 	case FTcpClient::ReadError:
 		UE_LOG(LogCPSControlCommunication, Verbose, TEXT("Unrecoverable FTcpClient error occurred when receiving control packet header: '%s' Code: %d."), *Error.GetMessage(), Error.GetCode())
 		Stop();
+		break;
 	default:
 		UE_LOG(LogCPSControlCommunication, Verbose, TEXT("Unhandled FTcpClient error occurred when receiving control packet header: '%s' Code: %d"), *Error.GetMessage(), Error.GetCode())
 		break;
diff --git a/Engine/Plugins/VirtualProduction/CaptureManager/CaptureManagerCore/Source/CaptureProtocolStack/Private/Control/Communication/ControlPacket.cpp b/Engine/Plugins/VirtualProduction/CaptureManager/CaptureManagerCore/Source/CaptureProtocolStack/Private/Control/Communication/ControlPacket.cpp
--- a/Engine/Plugins/VirtualProduction/CaptureManager/CaptureManagerCore/Source/CaptureProtocolStack/Private/Control/Communication/ControlPacket.cpp
+++ b/Engine/Plugins/VirtualProduction/CaptureManager/CaptureManagerCore/Source/CaptureProtocolStack/Private/Control/Communication/ControlPacket.cpp
@@ -11,9 +11,9 @@ const TArray<uint8> FControlPacketHeader::Header = { 'C', 'P', 'S', 'C', 'O', 'N
 
 #define CPS_CHECK_RESULT_RET_ERROR(Result) if (Result.HasError()) { return Result.StealError(); }
 
-TProtocolResult<FControlPacketHeader> FControlPacketHeader::Deserialize(ITcpSocketReader& InReader)
+TProtocolResult<FControlPacketHeader> FControlPacketHeader::Deserialize(ITcpSocketReader& InReader, uint32 InTimeoutForControlHeader)
 {
-    TProtocolResult<TArray<uint8>> HeaderResult = InReader.ReceiveMessage(Header.Num());
+    TProtocolResult<TArray<uint8>> HeaderResult = InReader.ReceiveMessage(Header.Num(), InTimeoutForControlHeader);
     CPS_CHECK_RESULT_RET_ERROR(HeaderResult);
 
     TArray<uint8> HeaderData = HeaderResult.StealValue();
diff --git a/Engine/Plugins/VirtualProduction/CaptureManager/CaptureManagerCore/Source/CaptureProtocolStack/Private/Control/ControlMessenger.cpp b/Engine/Plugins/VirtualProduction/CaptureManager/CaptureManagerCore/Source/CaptureProtocolStack/Private/Control/ControlMessenger.cpp
--- a/Engine/Plugins/VirtualProduction/CaptureManager/CaptureManagerCore/Source/CaptureProtocolStack/Private/Control/ControlMessenger.cpp
+++ b/Engine/Plugins/VirtualProduction/CaptureManager/CaptureManagerCore/Source/CaptureProtocolStack/Private/Control/ControlMessenger.cpp
@@ -170,6 +170,7 @@ void FControlMessenger::KeepAlive()
 	}));
 }
 
+// Note: This callback happens on sender and receiver thread
 void FControlMessenger::MessageHandler(FControlPacket InPacket)
 {
 	TProtocolResult<FControlMessage> DeserializeResult = FControlMessage::Deserialize(InPacket);
diff --git a/Engine/Plugins/VirtualProduction/CaptureManager/CaptureManagerCore/Source/CaptureProtocolStack/Public/Control/Communication/ControlCommunication.h b/Engine/Plugins/VirtualProduction/CaptureManager/CaptureManagerCore/Source/CaptureProtocolStack/Public/Control/Communication/ControlCommunication.h
--- a/Engine/Plugins/VirtualProduction/CaptureManager/CaptureManagerCore/Source/CaptureProtocolStack/Public/Control/Communication/ControlCommunication.h
+++ b/Engine/Plugins/VirtualProduction/CaptureManager/CaptureManagerCore/Source/CaptureProtocolStack/Public/Control/Communication/ControlCommunication.h
@@ -23,8 +23,6 @@ public:
 	DECLARE_DELEGATE_OneParam(FOnPacketReceived, FControlPacket InPacket)
 	DECLARE_DELEGATE(FCommunicationStoppedHandler)
 
-	using FRunnerType = TQueueRunner<FControlPacket>;
-
 	FControlCommunication();
 
 	TProtocolResult<void> Init();
@@ -62,11 +60,13 @@ private:
 	};
 
 
-	void OnProcessReceivedPacket(FControlPacket InMessage);
-	void OnProcessSentPacket(FControlPacket InMessage);
 	void RunnableHandler(TSharedPtr<FRunnable> InRunnable);
 	void CommunicationRunnableStopped();
 
+	void SendNextPacket();
+	void SendPacket(FControlPacket InPacket);
+	void ReceiveNextPacket();
+
 	TProtocolResult<FControlPacketHeader> ReceiveControlHeader();
 	TProtocolResult<void> ReceiveControlPacket(const FControlPacketHeader& InHeader);
 
@@ -74,8 +74,9 @@ private:
 	FOnPacketReceived ReceiveHandler;
 	FCommunicationStoppedHandler CommunicationStoppedHandler;
 
-	FRunnerType SynchronizedReceiver;
-	FRunnerType SynchronizedSender;
+	TQueue<FControlPacket> SenderQueue;
+	std::atomic_int SenderQueueCurrentSize = 0;
+	TQueue<FControlPacket> ReceiverQueue;
 	TQueueRunner<TSharedPtr<FRunnable>> SynchronizedRunnable;
 	TWeakPtr<FCommunicationRunnable> Runnable;
 };
diff --git a/Engine/Plugins/VirtualProduction/CaptureManager/CaptureManagerCore/Source/CaptureProtocolStack/Public/Control/Communication/ControlPacket.h b/Engine/Plugins/VirtualProduction/CaptureManager/CaptureManagerCore/Source/CaptureProtocolStack/Public/Control/Communication/ControlPacket.h
--- a/Engine/Plugins/VirtualProduction/CaptureManager/CaptureManagerCore/Source/CaptureProtocolStack/Public/Control/Communication/ControlPacket.h
+++ b/Engine/Plugins/VirtualProduction/CaptureManager/CaptureManagerCore/Source/CaptureProtocolStack/Public/Control/Communication/ControlPacket.h
@@ -19,7 +19,8 @@ public:
 
 	static const TArray<uint8> Header;
 
-	static TProtocolResult<FControlPacketHeader> Deserialize(ITcpSocketReader& InReader);
+	static TProtocolResult<FControlPacketHeader> Deserialize(ITcpSocketReader& InReader, 
+															 uint32 InTimeoutForControlHeader = ITcpSocketReader::DefaultWaitTimeoutMs);
 	static uint32 GetHeaderSize();
 
 	FControlPacketHeader();
