Update rtFact device to use LiveLinkFace device

#Synced-CL 51705988

diff --git a/Engine/Plugins/VirtualProduction/CaptureManager/CaptureManagerCore/Source/CaptureProtocolStack/Private/Control/Capabilities/Utils.cpp b/Engine/Plugins/VirtualProduction/CaptureManager/CaptureManagerCore/Source/CaptureProtocolStack/Private/Control/Capabilities/Utils.cpp
--- a/Engine/Plugins/VirtualProduction/CaptureManager/CaptureManagerCore/Source/CaptureProtocolStack/Private/Control/Capabilities/Utils.cpp
+++ b/Engine/Plugins/VirtualProduction/CaptureManager/CaptureManagerCore/Source/CaptureProtocolStack/Private/Control/Capabilities/Utils.cpp
@@ -51,9 +51,9 @@ TProtocolResult<void> FCapabilityUtilities::ValueFromJsonObject(const TSharedPtr
 	}
 	else if (InType == ECapabilityValue::Type::Integer)
 	{
-		int Value;
+		int32 Value;
 		CHECK_PARSE(FJsonUtility::ParseNumber(InJsonObject, InKey, Value));
-		OutValue.Set<int>(Value);
+		OutValue.Set<int32>(Value);
 	}
 	else if (InType == ECapabilityValue::Type::Float)
 	{
@@ -87,7 +87,7 @@ void FCapabilityUtilities::SetProperJsonFieldType(const FString& InKey,
 	}
 	else if (Type == ECapabilityValue::Type::Integer)
 	{
-		OutJsonObjectPtr->SetNumberField(InKey, InValue.Get<int>());
+		OutJsonObjectPtr->SetNumberField(InKey, InValue.Get<int32>());
 	}
 	else if (Type == ECapabilityValue::Type::Float)
 	{
diff --git a/Engine/Plugins/VirtualProduction/CaptureManager/CaptureManagerCore/Source/CaptureProtocolStack/Private/Control/Messages/ControlResponse.cpp b/Engine/Plugins/VirtualProduction/CaptureManager/CaptureManagerCore/Source/CaptureProtocolStack/Private/Control/Messages/ControlResponse.cpp
--- a/Engine/Plugins/VirtualProduction/CaptureManager/CaptureManagerCore/Source/CaptureProtocolStack/Private/Control/Messages/ControlResponse.cpp
+++ b/Engine/Plugins/VirtualProduction/CaptureManager/CaptureManagerCore/Source/CaptureProtocolStack/Private/Control/Messages/ControlResponse.cpp
@@ -451,14 +451,14 @@ TProtocolResult<FCapability> FListAllCapabilitiesResponse::ParseCapability(const
 
 		if (JsonProperty->HasField(CPS::Capabilities::GMin))
 		{
-			TVariant<int, float> MinVal;
+			TVariant<int32, float> MinVal;
 
 			float MinLimit;
 			CHECK_PARSE(FJsonUtility::ParseNumber(JsonProperty, CPS::Capabilities::GMin, MinLimit));
 
 			if (Property.Type == ECapabilityValue::Type::Integer)
 			{
-				MinVal.Set<int>(MinLimit);
+				MinVal.Set<int32>(MinLimit);
 			}
 			else
 			{
@@ -469,14 +469,14 @@ TProtocolResult<FCapability> FListAllCapabilitiesResponse::ParseCapability(const
 
 		if (JsonProperty->HasField(CPS::Capabilities::GMax))
 		{
-			TVariant<int, float> MaxVal;
+			TVariant<int32, float> MaxVal;
 
 			float MaxLimit;
 			CHECK_PARSE(FJsonUtility::ParseNumber(JsonProperty, CPS::Capabilities::GMax, MaxLimit));
 
 			if (Property.Type == ECapabilityValue::Type::Integer)
 			{
-				MaxVal.Set<int>(MaxLimit);
+				MaxVal.Set<int32>(MaxLimit);
 			}
 			else
 			{
diff --git a/Engine/Plugins/VirtualProduction/CaptureManager/CaptureManagerCore/Source/CaptureProtocolStack/Public/Control/Capabilities/Capability.h b/Engine/Plugins/VirtualProduction/CaptureManager/CaptureManagerCore/Source/CaptureProtocolStack/Public/Control/Capabilities/Capability.h
--- a/Engine/Plugins/VirtualProduction/CaptureManager/CaptureManagerCore/Source/CaptureProtocolStack/Public/Control/Capabilities/Capability.h
+++ b/Engine/Plugins/VirtualProduction/CaptureManager/CaptureManagerCore/Source/CaptureProtocolStack/Public/Control/Capabilities/Capability.h
@@ -31,7 +31,7 @@ enum Type
 };
 } // namespace ECapabilityAccess
 
-using FCapabilityValue = TVariant<bool, int, float, FString>;
+using FCapabilityValue = TVariant<bool, int32, float, FString>;
 
 using FCapabilityValues = TMap<FString, FCapabilityValue>;
 
@@ -41,8 +41,8 @@ struct FCapabilityProperty
 	FString Name;
 	ECapabilityValue::Type Type;
 	ECapabilityAccess::Type Access;
-	TOptional<TVariant<int, float>> Min;
-	TOptional<TVariant<int, float>> Max;
+	TOptional<TVariant<int32, float>> Min;
+	TOptional<TVariant<int32, float>> Max;
 	TOptional<TArray<FString>> EnumOptions;
 	FCapabilityValue CurrentValue;
 };
diff --git a/Engine/Plugins/VirtualProduction/CaptureManager/CaptureManagerCore/Source/CaptureProtocolStack/Public/Control/Messages/Constants.h b/Engine/Plugins/VirtualProduction/CaptureManager/CaptureManagerCore/Source/CaptureProtocolStack/Public/Control/Messages/Constants.h
--- a/Engine/Plugins/VirtualProduction/CaptureManager/CaptureManagerCore/Source/CaptureProtocolStack/Public/Control/Messages/Constants.h
+++ b/Engine/Plugins/VirtualProduction/CaptureManager/CaptureManagerCore/Source/CaptureProtocolStack/Public/Control/Messages/Constants.h
@@ -35,6 +35,18 @@ extern CAPTUREPROTOCOLSTACK_API const TCHAR* GTakeAdded;
 extern CAPTUREPROTOCOLSTACK_API const TCHAR* GTakeRemoved;
 extern CAPTUREPROTOCOLSTACK_API const TCHAR* GTakeUpdated;
 
+// Capabilities
+extern CAPTUREPROTOCOLSTACK_API const TCHAR* GCapabilitiesRoot;
+extern CAPTUREPROTOCOLSTACK_API const TCHAR* GCapabilitiesListAll;
+extern CAPTUREPROTOCOLSTACK_API const TCHAR* GCapabilityProperties;
+extern CAPTUREPROTOCOLSTACK_API const TCHAR* GCapabilityPropertySet;
+extern CAPTUREPROTOCOLSTACK_API const TCHAR* GCapabilityPropertyGet;
+extern CAPTUREPROTOCOLSTACK_API const TCHAR* GCapabilityCommands;
+extern CAPTUREPROTOCOLSTACK_API const TCHAR* GCapabilityCommandExecute;
+extern CAPTUREPROTOCOLSTACK_API const TCHAR* GCapabilityEvents;
+extern CAPTUREPROTOCOLSTACK_API const TCHAR* GCapabilityEventSubscribe;
+extern CAPTUREPROTOCOLSTACK_API const TCHAR* GCapabilityEventUnsubscribe;
+
 // Platform Specific: IOS
 extern CAPTUREPROTOCOLSTACK_API const TCHAR* GDiskCapacity;
 extern CAPTUREPROTOCOLSTACK_API const TCHAR* GBattery;
@@ -136,4 +148,37 @@ extern CAPTUREPROTOCOLSTACK_API const TCHAR* GRequestParsingFailed;
 extern CAPTUREPROTOCOLSTACK_API const TCHAR* GNotSupported;
 extern CAPTUREPROTOCOLSTACK_API const TCHAR* GNoSession;
 extern CAPTUREPROTOCOLSTACK_API const TCHAR* GMissingBody;
-}
\ No newline at end of file
+}
+
+namespace UE::CaptureManager::CPS::Capabilities
+{
+extern CAPTUREPROTOCOLSTACK_API const TCHAR* GId;
+extern CAPTUREPROTOCOLSTACK_API const TCHAR* GName;
+extern CAPTUREPROTOCOLSTACK_API const TCHAR* GValue;
+extern CAPTUREPROTOCOLSTACK_API const TCHAR* GProperties;
+extern CAPTUREPROTOCOLSTACK_API const TCHAR* GType;
+extern CAPTUREPROTOCOLSTACK_API const TCHAR* GAccess;
+extern CAPTUREPROTOCOLSTACK_API const TCHAR* GReadOnly;
+extern CAPTUREPROTOCOLSTACK_API const TCHAR* GReadWrite;
+extern CAPTUREPROTOCOLSTACK_API const TCHAR* GEnumOptions;
+extern CAPTUREPROTOCOLSTACK_API const TCHAR* GCurrentValue;
+extern CAPTUREPROTOCOLSTACK_API const TCHAR* GMin;
+extern CAPTUREPROTOCOLSTACK_API const TCHAR* GMax;
+extern CAPTUREPROTOCOLSTACK_API const TCHAR* GGet;
+extern CAPTUREPROTOCOLSTACK_API const TCHAR* GSet;
+extern CAPTUREPROTOCOLSTACK_API const TCHAR* GCommands;
+extern CAPTUREPROTOCOLSTACK_API const TCHAR* GParameters;
+extern CAPTUREPROTOCOLSTACK_API const TCHAR* GExecute;
+extern CAPTUREPROTOCOLSTACK_API const TCHAR* GOptional;
+extern CAPTUREPROTOCOLSTACK_API const TCHAR* GReturnType;
+extern CAPTUREPROTOCOLSTACK_API const TCHAR* GDefaultValue;
+extern CAPTUREPROTOCOLSTACK_API const TCHAR* GEvents;
+extern CAPTUREPROTOCOLSTACK_API const TCHAR* GArguments;
+extern CAPTUREPROTOCOLSTACK_API const TCHAR* GSubscribe;
+extern CAPTUREPROTOCOLSTACK_API const TCHAR* GUnsubscribe;
+extern CAPTUREPROTOCOLSTACK_API const TCHAR* GBoolean;
+extern CAPTUREPROTOCOLSTACK_API const TCHAR* GInteger;
+extern CAPTUREPROTOCOLSTACK_API const TCHAR* GFloat;
+extern CAPTUREPROTOCOLSTACK_API const TCHAR* GString;
+extern CAPTUREPROTOCOLSTACK_API const TCHAR* GEnumeration;
+} // namespace UE::CaptureManager::CPS::Capabilities
\ No newline at end of file
