[CaptureManager] Add CaptureMetadataExtraction module with shared extractors for TakeArchive, MonoVideo, StereoVideo, and LiveLinkFace metadata.

Delete the now-dead LiveLinkFaceMetadata and StereoCameraMetadata modules from CaptureManagerApp and move pixel format and rotation enums to CaptureDataConverter
Consolidate asset preparation, path resolution, and capture data asset creation into shared helpers in CaptureManagerEditor.

#Synced-CL 52674001

diff --git a/Engine/Plugins/VirtualProduction/CaptureManager/CaptureManagerApp/CaptureManagerApp.uplugin b/Engine/Plugins/VirtualProduction/CaptureManager/CaptureManagerApp/CaptureManagerApp.uplugin
--- a/Engine/Plugins/VirtualProduction/CaptureManager/CaptureManagerApp/CaptureManagerApp.uplugin
+++ b/Engine/Plugins/VirtualProduction/CaptureManager/CaptureManagerApp/CaptureManagerApp.uplugin
@@ -23,16 +23,6 @@
 			"Type": "Editor",
 			"LoadingPhase": "Default"
 		},
-		{
-			"Name": "LiveLinkFaceMetadata",
-			"Type": "Editor",
-			"LoadingPhase": "Default"
-		},
-		{
-			"Name": "StereoCameraMetadata",
-			"Type": "Editor",
-			"LoadingPhase": "Default"
-		},
 		{
 			"Name": "CaptureManagerUnrealEndpoint",
 			"Type": "Editor",
diff --git a/Engine/Plugins/VirtualProduction/CaptureManager/CaptureManagerApp/Config/Engine.ini b/Engine/Plugins/VirtualProduction/CaptureManager/CaptureManagerApp/Config/Engine.ini
--- a/Engine/Plugins/VirtualProduction/CaptureManager/CaptureManagerApp/Config/Engine.ini
+++ b/Engine/Plugins/VirtualProduction/CaptureManager/CaptureManagerApp/Config/Engine.ini
@@ -1,2 +1,6 @@
+[CoreRedirects]
++EnumRedirects=(OldName="/Script/LiveLinkCapabilities.EIngestCapability_ImagePixelFormat", NewName="/Script/CaptureDataConverter.ECaptureManagerPixelFormat")
++EnumRedirects=(OldName="/Script/LiveLinkCapabilities.EIngestCapability_ImageRotation", NewName="/Script/CaptureDataConverter.ECaptureManagerRotation")
+
 [/Script/LiveLinkHub.LiveLinkHubSettings]
 +ProjectSettingsToDisplay=Capture Manager
\ No newline at end of file
diff --git a/Engine/Plugins/VirtualProduction/CaptureManager/CaptureManagerApp/Source/CaptureManagerEditor/Private/IngestManagement/IngestJobExecutor.cpp b/Engine/Plugins/VirtualProduction/CaptureManager/CaptureManagerApp/Source/CaptureManagerEditor/Private/IngestManagement/IngestJobExecutor.cpp
--- a/Engine/Plugins/VirtualProduction/CaptureManager/CaptureManagerApp/Source/CaptureManagerEditor/Private/IngestManagement/IngestJobExecutor.cpp
+++ b/Engine/Plugins/VirtualProduction/CaptureManager/CaptureManagerApp/Source/CaptureManagerEditor/Private/IngestManagement/IngestJobExecutor.cpp
@@ -26,31 +26,31 @@ DEFINE_LOG_CATEGORY_STATIC(LogCaptureManagerIngestJobExecutor, Log, All);
 namespace UE::CaptureManager
 {
 
-static EIngestCapability_ImagePixelFormat ConvertImagePixelFormat(::EImagePixelFormat InImagePixelFormat)
+static ECaptureManagerPixelFormat ConvertImagePixelFormat(::EImagePixelFormat InImagePixelFormat)
 {
 	switch (InImagePixelFormat)
 	{
 		case ::EImagePixelFormat::U8_Mono:
-			return EIngestCapability_ImagePixelFormat::U8_Mono;
+			return ECaptureManagerPixelFormat::U8_Mono;
 		case ::EImagePixelFormat::U8_BGRA:
 		default:
-			return EIngestCapability_ImagePixelFormat::U8_BGRA;
+			return ECaptureManagerPixelFormat::U8_BGRA;
 	}
 }
 
-static EIngestCapability_ImageRotation ConvertImageRotation(::EImageRotation InImageRotation)
+static ECaptureManagerRotation ConvertImageRotation(::EImageRotation InImageRotation)
 {
 	switch (InImageRotation)
 	{
 		case ::EImageRotation::CW_90:
-			return EIngestCapability_ImageRotation::CW_90;
+			return ECaptureManagerRotation::CW_90;
 		case ::EImageRotation::CW_180:
-			return EIngestCapability_ImageRotation::CW_180;
+			return ECaptureManagerRotation::CW_180;
 		case ::EImageRotation::CW_270:
-			return EIngestCapability_ImageRotation::CW_270;
+			return ECaptureManagerRotation::CW_270;
 		case ::EImageRotation::None:
 		default:
-			return EIngestCapability_ImageRotation::None;
+			return ECaptureManagerRotation::None;
 	}
 }
 
diff --git a/Engine/Plugins/VirtualProduction/CaptureManager/CaptureManagerApp/Source/IngestLiveLinkDevice/Private/BaseIngestLiveLinkDevice.cpp b/Engine/Plugins/VirtualProduction/CaptureManager/CaptureManagerApp/Source/IngestLiveLinkDevice/Private/BaseIngestLiveLinkDevice.cpp
--- a/Engine/Plugins/VirtualProduction/CaptureManager/CaptureManagerApp/Source/IngestLiveLinkDevice/Private/BaseIngestLiveLinkDevice.cpp
+++ b/Engine/Plugins/VirtualProduction/CaptureManager/CaptureManagerApp/Source/IngestLiveLinkDevice/Private/BaseIngestLiveLinkDevice.cpp
@@ -443,8 +443,8 @@ FCaptureDataConverterResult<void> UBaseIngestLiveLinkDevice::FImpl::RunConversio
 	FCaptureConvertVideoOutputParams VideoParams;
 	VideoParams.Format = InIngestOptions->Video.Format;
 	VideoParams.ImageFileName = InIngestOptions->Video.FileNamePrefix;
-	VideoParams.OutputPixelFormat = Private::ConvertPixelFormat(InIngestOptions->Video.PixelFormat);
-	VideoParams.Rotation = Private::ConvertImageRotation(InIngestOptions->Video.Rotation);
+	VideoParams.OutputPixelFormat = ToMediaPixelFormat(InIngestOptions->Video.PixelFormat);
+	VideoParams.Rotation = ToMediaOrientation(InIngestOptions->Video.Rotation);
 
 	InParams.VideoOutputParams = MoveTemp(VideoParams);
 
@@ -456,7 +456,7 @@ FCaptureDataConverterResult<void> UBaseIngestLiveLinkDevice::FImpl::RunConversio
 
 	FCaptureConvertDepthOutputParams DepthParams;
 	DepthParams.ImageFileName = TEXT("depth");
-	DepthParams.Rotation = Private::ConvertImageRotation(InIngestOptions->Video.Rotation); // Apply the same rotation as video to depth stream
+	DepthParams.Rotation = ToMediaOrientation(InIngestOptions->Video.Rotation); // Apply the same rotation as video to depth stream
 
 	InParams.DepthOutputParams = MoveTemp(DepthParams);
 
diff --git a/Engine/Plugins/VirtualProduction/CaptureManager/CaptureManagerApp/Source/IngestLiveLinkDevice/Private/Utils/CaptureExtractTimecode.cpp b/Engine/Plugins/VirtualProduction/CaptureManager/CaptureManagerApp/Source/IngestLiveLinkDevice/Private/Utils/CaptureExtractTimecode.cpp
deleted file mode 100644
diff --git a/Engine/Plugins/VirtualProduction/CaptureManager/CaptureManagerApp/Source/IngestLiveLinkDevice/Private/Utils/IngestLiveLinkDeviceUtils.cpp b/Engine/Plugins/VirtualProduction/CaptureManager/CaptureManagerApp/Source/IngestLiveLinkDevice/Private/Utils/IngestLiveLinkDeviceUtils.cpp
--- a/Engine/Plugins/VirtualProduction/CaptureManager/CaptureManagerApp/Source/IngestLiveLinkDevice/Private/Utils/IngestLiveLinkDeviceUtils.cpp
+++ b/Engine/Plugins/VirtualProduction/CaptureManager/CaptureManagerApp/Source/IngestLiveLinkDevice/Private/Utils/IngestLiveLinkDeviceUtils.cpp
@@ -5,58 +5,6 @@
 namespace UE::CaptureManager
 {
 
-namespace Private
-{
-
-EMediaTexturePixelFormat ConvertPixelFormat(EIngestCapability_ImagePixelFormat InImagePixelFormat)
-{
-	switch (InImagePixelFormat)
-	{
-		case EIngestCapability_ImagePixelFormat::U8_RGB:
-			return EMediaTexturePixelFormat::U8_RGB;
-		case EIngestCapability_ImagePixelFormat::U8_BGR:
-			return EMediaTexturePixelFormat::U8_BGR;
-		case EIngestCapability_ImagePixelFormat::U8_RGBA:
-			return EMediaTexturePixelFormat::U8_RGBA;
-		case EIngestCapability_ImagePixelFormat::U8_BGRA:
-			return EMediaTexturePixelFormat::U8_BGRA;
-		case EIngestCapability_ImagePixelFormat::U8_I444:
-			return EMediaTexturePixelFormat::U8_I444;
-		case EIngestCapability_ImagePixelFormat::U8_I420:
-			return EMediaTexturePixelFormat::U8_I420;
-		case EIngestCapability_ImagePixelFormat::U8_NV12:
-			return EMediaTexturePixelFormat::U8_NV12;
-		case EIngestCapability_ImagePixelFormat::U8_Mono:
-			return EMediaTexturePixelFormat::U8_Mono;
-		case EIngestCapability_ImagePixelFormat::U16_Mono:
-			return EMediaTexturePixelFormat::U16_Mono;
-		case EIngestCapability_ImagePixelFormat::F_Mono:
-			return EMediaTexturePixelFormat::F_Mono;
-		case EIngestCapability_ImagePixelFormat::Undefined:
-		default:
-			check(false);
-			return EMediaTexturePixelFormat::Undefined;
-	}
-}
-
-EMediaOrientation ConvertImageRotation(EIngestCapability_ImageRotation InImageRotation)
-{
-	switch (InImageRotation)
-	{
-		case EIngestCapability_ImageRotation::CW_90:
-			return EMediaOrientation::CW90;
-		case EIngestCapability_ImageRotation::CW_180:
-			return EMediaOrientation::CW180;
-		case EIngestCapability_ImageRotation::CW_270:
-			return EMediaOrientation::CW270;
-		case EIngestCapability_ImageRotation::None:
-		default:
-			return EMediaOrientation::Original;
-	}
-}
-
-}
-
 FString ErrorOriginToString(FTakeMetadataParserError::EOrigin InOrigin)
 {
 	switch (InOrigin)
diff --git a/Engine/Plugins/VirtualProduction/CaptureManager/CaptureManagerApp/Source/IngestLiveLinkDevice/Public/Utils/CaptureExtractTimecode.h b/Engine/Plugins/VirtualProduction/CaptureManager/CaptureManagerApp/Source/IngestLiveLinkDevice/Public/Utils/CaptureExtractTimecode.h
deleted file mode 100644
diff --git a/Engine/Plugins/VirtualProduction/CaptureManager/CaptureManagerApp/Source/IngestLiveLinkDevice/Public/Utils/IngestLiveLinkDeviceUtils.h b/Engine/Plugins/VirtualProduction/CaptureManager/CaptureManagerApp/Source/IngestLiveLinkDevice/Public/Utils/IngestLiveLinkDeviceUtils.h
--- a/Engine/Plugins/VirtualProduction/CaptureManager/CaptureManagerApp/Source/IngestLiveLinkDevice/Public/Utils/IngestLiveLinkDeviceUtils.h
+++ b/Engine/Plugins/VirtualProduction/CaptureManager/CaptureManagerApp/Source/IngestLiveLinkDevice/Public/Utils/IngestLiveLinkDeviceUtils.h
@@ -3,16 +3,9 @@
 #pragma once
 
 #include "CaptureManagerTakeMetadata.h"
-#include "Ingest/IngestCapability_Options.h"
-#include "MediaSample.h"
 
 namespace UE::CaptureManager
 {
-namespace Private
-{
-EMediaTexturePixelFormat ConvertPixelFormat(EIngestCapability_ImagePixelFormat InImagePixelFormat);
-EMediaOrientation ConvertImageRotation(EIngestCapability_ImageRotation InImageRotation);
-}
 
 INGESTLIVELINKDEVICE_API FString ErrorOriginToString(FTakeMetadataParserError::EOrigin InOrigin);
 
diff --git a/Engine/Plugins/VirtualProduction/CaptureManager/CaptureManagerApp/Source/LiveLinkCapabilities/LiveLinkCapabilities.Build.cs b/Engine/Plugins/VirtualProduction/CaptureManager/CaptureManagerApp/Source/LiveLinkCapabilities/LiveLinkCapabilities.Build.cs
--- a/Engine/Plugins/VirtualProduction/CaptureManager/CaptureManagerApp/Source/LiveLinkCapabilities/LiveLinkCapabilities.Build.cs
+++ b/Engine/Plugins/VirtualProduction/CaptureManager/CaptureManagerApp/Source/LiveLinkCapabilities/LiveLinkCapabilities.Build.cs
@@ -15,7 +15,8 @@ public class LiveLinkCapabilities : ModuleRules
 			"CaptureUtils",
 			"Core",
 			"CaptureManagerUnrealEndpoint", // TODO: Only needed for the ULiveLinkHubIngestClientInfo definition, move it and remove this dependency
-			"CaptureManagerTakeMetadata"
+			"CaptureManagerTakeMetadata",
+			"CaptureDataConverter"
 		});
 
 		PrivateDependencyModuleNames.AddRange(new string[]
diff --git a/Engine/Plugins/VirtualProduction/CaptureManager/CaptureManagerApp/Source/LiveLinkCapabilities/Public/Ingest/IngestCapability_Options.h b/Engine/Plugins/VirtualProduction/CaptureManager/CaptureManagerApp/Source/LiveLinkCapabilities/Public/Ingest/IngestCapability_Options.h
--- a/Engine/Plugins/VirtualProduction/CaptureManager/CaptureManagerApp/Source/LiveLinkCapabilities/Public/Ingest/IngestCapability_Options.h
+++ b/Engine/Plugins/VirtualProduction/CaptureManager/CaptureManagerApp/Source/LiveLinkCapabilities/Public/Ingest/IngestCapability_Options.h
@@ -3,37 +3,13 @@
 #pragma once
 
 #include "LiveLinkDevice.h"
+#include "CaptureManagerConversionTypes.h"
 
 #include "IngestCapability_Options.generated.h"
 
-UENUM()
-enum class EIngestCapability_ImagePixelFormat
-{
-	Undefined = -1		UMETA(Hidden),
-	U8_RGB = 0			UMETA(DisplayName = "RGB 8-bit"),
-	U8_BGR				UMETA(DisplayName = "BGR 8-bit"),
-	U8_RGBA				UMETA(DisplayName = "RGBA 8-bit"),
-	U8_BGRA				UMETA(DisplayName = "BGRA 8-bit"),
-	U8_I444				UMETA(DisplayName = "I444 8-bit"),
-	U8_I420				UMETA(DisplayName = "I420 8-bit"),
-	U8_NV12				UMETA(DisplayName = "NV12 8-bit"),
-	U8_Mono				UMETA(DisplayName = "Mono 8-bit"),
-	U16_Mono			UMETA(DisplayName = "Mono 16-bit"),
-	F_Mono				UMETA(DisplayName = "Mono Float"),
-
-	Default = U8_RGB	UMETA(DisplayName = "Default (U8_RGB)")
-};
-
-UENUM()
-enum class EIngestCapability_ImageRotation
-{
-	None = 0			UMETA(DisplayName = "None"),
-	CW_90				UMETA(DisplayName = "Clockwise 90°"),
-	CW_180				UMETA(DisplayName = "Clockwise 180°"),
-	CW_270				UMETA(DisplayName = "Clockwise 270°"),
-
-	Default = None		UMETA(DisplayName = "Default (None)")
-};
+// Deprecated type aliases for backward compatibility.
+using EIngestCapability_ImagePixelFormat UE_DEPRECATED(5.8, "Use ECaptureManagerPixelFormat instead.") = ECaptureManagerPixelFormat;
+using EIngestCapability_ImageRotation UE_DEPRECATED(5.8, "Use ECaptureManagerRotation instead.") = ECaptureManagerRotation;
 
 USTRUCT(BlueprintType)
 struct LIVELINKCAPABILITIES_API FIngestCapability_VideoOptions
@@ -47,10 +23,10 @@ struct LIVELINKCAPABILITIES_API FIngestCapability_VideoOptions
 	FString Format;
 
 	UPROPERTY(EditAnywhere, BlueprintReadWrite, Category = "Live Link Device|Ingest")
-	EIngestCapability_ImagePixelFormat PixelFormat = EIngestCapability_ImagePixelFormat::U8_BGRA;
+	ECaptureManagerPixelFormat PixelFormat = ECaptureManagerPixelFormat::U8_BGRA;
 
 	UPROPERTY(EditAnywhere, BlueprintReadWrite, Category = "Live Link Device|Ingest")
-	EIngestCapability_ImageRotation Rotation = EIngestCapability_ImageRotation::None;
+	ECaptureManagerRotation Rotation = ECaptureManagerRotation::None;
 };
 
 USTRUCT(BlueprintType)
diff --git a/Engine/Plugins/VirtualProduction/CaptureManager/CaptureManagerApp/Source/LiveLinkFaceMetadata/LiveLinkFaceMetadata.Build.cs b/Engine/Plugins/VirtualProduction/CaptureManager/CaptureManagerApp/Source/LiveLinkFaceMetadata/LiveLinkFaceMetadata.Build.cs
deleted file mode 100644
diff --git a/Engine/Plugins/VirtualProduction/CaptureManager/CaptureManagerApp/Source/LiveLinkFaceMetadata/Private/LiveLinkFaceMetadata.cpp b/Engine/Plugins/VirtualProduction/CaptureManager/CaptureManagerApp/Source/LiveLinkFaceMetadata/Private/LiveLinkFaceMetadata.cpp
deleted file mode 100644
diff --git a/Engine/Plugins/VirtualProduction/CaptureManager/CaptureManagerApp/Source/LiveLinkFaceMetadata/Private/LiveLinkFaceMetadataModule.cpp b/Engine/Plugins/VirtualProduction/CaptureManager/CaptureManagerApp/Source/LiveLinkFaceMetadata/Private/LiveLinkFaceMetadataModule.cpp
deleted file mode 100644
diff --git a/Engine/Plugins/VirtualProduction/CaptureManager/CaptureManagerApp/Source/LiveLinkFaceMetadata/Private/LiveLinkFaceMetadataModule.h b/Engine/Plugins/VirtualProduction/CaptureManager/CaptureManagerApp/Source/LiveLinkFaceMetadata/Private/LiveLinkFaceMetadataModule.h
deleted file mode 100644
diff --git a/Engine/Plugins/VirtualProduction/CaptureManager/CaptureManagerApp/Source/LiveLinkFaceMetadata/Private/Utils/AppleDeviceList.cpp b/Engine/Plugins/VirtualProduction/CaptureManager/CaptureManagerApp/Source/LiveLinkFaceMetadata/Private/Utils/AppleDeviceList.cpp
deleted file mode 100644
diff --git a/Engine/Plugins/VirtualProduction/CaptureManager/CaptureManagerApp/Source/LiveLinkFaceMetadata/Private/Utils/AppleDeviceList.h b/Engine/Plugins/VirtualProduction/CaptureManager/CaptureManagerApp/Source/LiveLinkFaceMetadata/Private/Utils/AppleDeviceList.h
deleted file mode 100644
diff --git a/Engine/Plugins/VirtualProduction/CaptureManager/CaptureManagerApp/Source/LiveLinkFaceMetadata/Public/LiveLinkFaceMetadata.h b/Engine/Plugins/VirtualProduction/CaptureManager/CaptureManagerApp/Source/LiveLinkFaceMetadata/Public/LiveLinkFaceMetadata.h
deleted file mode 100644
diff --git a/Engine/Plugins/VirtualProduction/CaptureManager/CaptureManagerApp/Source/StereoCameraMetadata/Private/ResolutionResolver.cpp b/Engine/Plugins/VirtualProduction/CaptureManager/CaptureManagerApp/Source/StereoCameraMetadata/Private/ResolutionResolver.cpp
deleted file mode 100644
diff --git a/Engine/Plugins/VirtualProduction/CaptureManager/CaptureManagerApp/Source/StereoCameraMetadata/Private/ResolutionResolver.h b/Engine/Plugins/VirtualProduction/CaptureManager/CaptureManagerApp/Source/StereoCameraMetadata/Private/ResolutionResolver.h
deleted file mode 100644
diff --git a/Engine/Plugins/VirtualProduction/CaptureManager/CaptureManagerApp/Source/StereoCameraMetadata/Private/StereoCameraMetadataModule.cpp b/Engine/Plugins/VirtualProduction/CaptureManager/CaptureManagerApp/Source/StereoCameraMetadata/Private/StereoCameraMetadataModule.cpp
deleted file mode 100644
diff --git a/Engine/Plugins/VirtualProduction/CaptureManager/CaptureManagerApp/Source/StereoCameraMetadata/Private/StereoCameraMetadataParseUtils.cpp b/Engine/Plugins/VirtualProduction/CaptureManager/CaptureManagerApp/Source/StereoCameraMetadata/Private/StereoCameraMetadataParseUtils.cpp
deleted file mode 100644
diff --git a/Engine/Plugins/VirtualProduction/CaptureManager/CaptureManagerApp/Source/StereoCameraMetadata/Private/StereoCameraTakeMetadata.cpp b/Engine/Plugins/VirtualProduction/CaptureManager/CaptureManagerApp/Source/StereoCameraMetadata/Private/StereoCameraTakeMetadata.cpp
deleted file mode 100644
diff --git a/Engine/Plugins/VirtualProduction/CaptureManager/CaptureManagerApp/Source/StereoCameraMetadata/Private/StereoCameraTakeMetadata.h b/Engine/Plugins/VirtualProduction/CaptureManager/CaptureManagerApp/Source/StereoCameraMetadata/Private/StereoCameraTakeMetadata.h
deleted file mode 100644
diff --git a/Engine/Plugins/VirtualProduction/CaptureManager/CaptureManagerApp/Source/StereoCameraMetadata/Public/StereoCameraMetadataParseUtils.h b/Engine/Plugins/VirtualProduction/CaptureManager/CaptureManagerApp/Source/StereoCameraMetadata/Public/StereoCameraMetadataParseUtils.h
deleted file mode 100644
diff --git a/Engine/Plugins/VirtualProduction/CaptureManager/CaptureManagerApp/Source/StereoCameraMetadata/StereoCameraMetadata.Build.cs b/Engine/Plugins/VirtualProduction/CaptureManager/CaptureManagerApp/Source/StereoCameraMetadata/StereoCameraMetadata.Build.cs
deleted file mode 100644
diff --git a/Engine/Plugins/VirtualProduction/CaptureManager/CaptureManagerCore/CaptureManagerCore.uplugin b/Engine/Plugins/VirtualProduction/CaptureManager/CaptureManagerCore/CaptureManagerCore.uplugin
--- a/Engine/Plugins/VirtualProduction/CaptureManager/CaptureManagerCore/CaptureManagerCore.uplugin
+++ b/Engine/Plugins/VirtualProduction/CaptureManager/CaptureManagerCore/CaptureManagerCore.uplugin
@@ -61,6 +61,11 @@
 			"Name": "CaptureManagerTakeMetadata",
 			"Type": "Editor",
 			"LoadingPhase": "Default"
+		},
+		{
+			"Name": "CaptureMetadataExtraction",
+			"Type": "Editor",
+			"LoadingPhase": "Default"
 		}
 	],
 	"Plugins": [
diff --git a/Engine/Plugins/VirtualProduction/CaptureManager/CaptureManagerCore/Source/CaptureDataConverter/Internal/CaptureManagerConversionTypes.cpp b/Engine/Plugins/VirtualProduction/CaptureManager/CaptureManagerCore/Source/CaptureDataConverter/Internal/CaptureManagerConversionTypes.cpp
new file mode 100644
--- /dev/null
+++ b/Engine/Plugins/VirtualProduction/CaptureManager/CaptureManagerCore/Source/CaptureDataConverter/Internal/CaptureManagerConversionTypes.cpp
@@ -0,0 +1,42 @@
+// Copyright Epic Games, Inc. All Rights Reserved.
+
+#include "CaptureManagerConversionTypes.h"
+
+namespace UE::CaptureManager
+{
+
+EMediaTexturePixelFormat ToMediaPixelFormat(ECaptureManagerPixelFormat InPixelFormat)
+{
+	switch (InPixelFormat)
+	{
+		case ECaptureManagerPixelFormat::U8_RGB:   return EMediaTexturePixelFormat::U8_RGB;
+		case ECaptureManagerPixelFormat::U8_BGR:   return EMediaTexturePixelFormat::U8_BGR;
+		case ECaptureManagerPixelFormat::U8_RGBA:  return EMediaTexturePixelFormat::U8_RGBA;
+		case ECaptureManagerPixelFormat::U8_BGRA:  return EMediaTexturePixelFormat::U8_BGRA;
+		case ECaptureManagerPixelFormat::U8_I444:  return EMediaTexturePixelFormat::U8_I444;
+		case ECaptureManagerPixelFormat::U8_I420:  return EMediaTexturePixelFormat::U8_I420;
+		case ECaptureManagerPixelFormat::U8_NV12:  return EMediaTexturePixelFormat::U8_NV12;
+		case ECaptureManagerPixelFormat::U8_Mono:  return EMediaTexturePixelFormat::U8_Mono;
+		case ECaptureManagerPixelFormat::U16_Mono: return EMediaTexturePixelFormat::U16_Mono;
+		case ECaptureManagerPixelFormat::F_Mono:   return EMediaTexturePixelFormat::F_Mono;
+		case ECaptureManagerPixelFormat::Undefined:
+		default:
+			check(false);
+			return EMediaTexturePixelFormat::Undefined;
+	}
+}
+
+EMediaOrientation ToMediaOrientation(ECaptureManagerRotation InRotation)
+{
+	switch (InRotation)
+	{
+		case ECaptureManagerRotation::CW_90:  return EMediaOrientation::CW90;
+		case ECaptureManagerRotation::CW_180: return EMediaOrientation::CW180;
+		case ECaptureManagerRotation::CW_270: return EMediaOrientation::CW270;
+		case ECaptureManagerRotation::None:
+		default:
+			return EMediaOrientation::Original;
+	}
+}
+
+} // namespace UE::CaptureManager
diff --git a/Engine/Plugins/VirtualProduction/CaptureManager/CaptureManagerCore/Source/CaptureDataConverter/Internal/CaptureManagerConversionTypes.h b/Engine/Plugins/VirtualProduction/CaptureManager/CaptureManagerCore/Source/CaptureDataConverter/Internal/CaptureManagerConversionTypes.h
new file mode 100644
--- /dev/null
+++ b/Engine/Plugins/VirtualProduction/CaptureManager/CaptureManagerCore/Source/CaptureDataConverter/Internal/CaptureManagerConversionTypes.h
@@ -0,0 +1,50 @@
+// Copyright Epic Games, Inc. All Rights Reserved.
+
+#pragma once
+
+#include "MediaSample.h"
+
+#include "CaptureManagerConversionTypes.generated.h"
+
+UENUM(BlueprintType)
+enum class ECaptureManagerPixelFormat : uint8
+{
+	// Packed RGB
+	U8_RGB = 0		UMETA(DisplayName = "RGB 8-bit"),
+	U8_BGR			UMETA(DisplayName = "BGR 8-bit"),
+	U8_RGBA			UMETA(DisplayName = "RGBA 8-bit"),
+	U8_BGRA			UMETA(DisplayName = "BGRA 8-bit"),
+	// Planar YUV — only supported when a third-party encoder (e.g. FFmpeg) is configured.
+	U8_I444			UMETA(DisplayName = "I444 8-bit"),
+	U8_I420			UMETA(DisplayName = "I420 8-bit"),
+	U8_NV12			UMETA(DisplayName = "NV12 8-bit"),
+	// Grayscale
+	U8_Mono			UMETA(DisplayName = "Mono 8-bit"),
+	U16_Mono		UMETA(DisplayName = "Mono 16-bit"),
+	F_Mono			UMETA(DisplayName = "Mono Float"),
+
+	Default = U8_RGB	UMETA(DisplayName = "Default (U8_RGB)"),
+	Undefined = 255		UMETA(Hidden),
+};
+
+UENUM(BlueprintType)
+enum class ECaptureManagerRotation : uint8
+{
+	None  = 0	UMETA(DisplayName = "None"),
+	CW_90		UMETA(DisplayName = "Clockwise 90°"),
+	CW_180		UMETA(DisplayName = "Clockwise 180°"),
+	CW_270		UMETA(DisplayName = "Clockwise 270°"),
+
+	Default = None	UMETA(DisplayName = "Default (None)"),
+};
+
+namespace UE::CaptureManager
+{
+
+/** Maps ECaptureManagerPixelFormat to the Core media pipeline type. */
+CAPTUREDATACONVERTER_API EMediaTexturePixelFormat ToMediaPixelFormat(ECaptureManagerPixelFormat InPixelFormat);
+
+/** Maps ECaptureManagerRotation to the Core media pipeline type. */
+CAPTUREDATACONVERTER_API EMediaOrientation ToMediaOrientation(ECaptureManagerRotation InRotation);
+
+} // namespace UE::CaptureManager
diff --git a/Engine/Plugins/VirtualProduction/CaptureManager/CaptureManagerCore/Source/CaptureMetadataExtraction/CaptureMetadataExtraction.Build.cs b/Engine/Plugins/VirtualProduction/CaptureManager/CaptureManagerCore/Source/CaptureMetadataExtraction/CaptureMetadataExtraction.Build.cs
new file mode 100644
--- /dev/null
+++ b/Engine/Plugins/VirtualProduction/CaptureManager/CaptureManagerCore/Source/CaptureMetadataExtraction/CaptureMetadataExtraction.Build.cs
@@ -0,0 +1,36 @@
+// Copyright Epic Games, Inc. All Rights Reserved.
+
+using UnrealBuildTool;
+using System.IO;
+using UnrealBuildTool.Rules;
+
+public class CaptureMetadataExtraction : ModuleRules
+{
+	public CaptureMetadataExtraction(ReadOnlyTargetRules Target) : base(Target)
+	{
+		PCHUsage = PCHUsageMode.UseExplicitOrSharedPCHs;
+
+		PublicDependencyModuleNames.AddRange(new string[]
+		{
+			"Core",
+			"CoreUObject",
+			"Engine",
+			"CaptureManagerTakeMetadata"
+		});
+
+		PrivateDependencyModuleNames.AddRange(new string[]
+		{
+			"CaptureUtils",
+			"Media",
+			"CaptureDataCore",
+			"CaptureManagerMediaRW",
+			"DataIngestCore",
+			"ElectraPlayerPlugin",
+			"CaptureDataUtils",
+			"Json",
+			"MediaAssets"
+		});
+
+		ShortName = "CapMetExt";
+	}
+}
diff --git a/Engine/Plugins/VirtualProduction/CaptureManager/CaptureManagerCore/Source/CaptureMetadataExtraction/Internal/CaptureManagerCalibrationUtils.cpp b/Engine/Plugins/VirtualProduction/CaptureManager/CaptureManagerCore/Source/CaptureMetadataExtraction/Internal/CaptureManagerCalibrationUtils.cpp
new file mode 100644
--- /dev/null
+++ b/Engine/Plugins/VirtualProduction/CaptureManager/CaptureManagerCore/Source/CaptureMetadataExtraction/Internal/CaptureManagerCalibrationUtils.cpp
@@ -0,0 +1,49 @@
+// Copyright Epic Games, Inc. All Rights Reserved.
+
+#include "CaptureManagerCalibrationUtils.h"
+
+#include "Misc/FileHelper.h"
+#include "Serialization/JsonSerializer.h"
+
+namespace UE::CaptureManager
+{
+
+FString DetectCalibrationFormat(const FString& InFilePath)
+{
+	FString JsonContent;
+	if (!FFileHelper::LoadFileToString(JsonContent, *InFilePath))
+	{
+		return FString();
+	}
+
+	TSharedRef<TJsonReader<TCHAR>> JsonReader = TJsonReaderFactory<>::Create(JsonContent);
+
+	// Try parsing as object (Unreal format)
+	TSharedPtr<FJsonObject> JsonObject;
+	if (FJsonSerializer::Deserialize(JsonReader, JsonObject) && JsonObject.IsValid())
+	{
+		if (JsonObject->HasField(TEXT("Calibrations")))
+		{
+			return TEXT("unreal");
+		}
+	}
+
+	// Try parsing as array (OpenCV format)
+	JsonReader = TJsonReaderFactory<>::Create(JsonContent);
+	TArray<TSharedPtr<FJsonValue>> JsonArray;
+	if (FJsonSerializer::Deserialize(JsonReader, JsonArray))
+	{
+		for (const TSharedPtr<FJsonValue>& Element : JsonArray)
+		{
+			const TSharedPtr<FJsonObject>* ElementObject = nullptr;
+			if (Element->TryGetObject(ElementObject) && (*ElementObject)->HasField(TEXT("fx")))
+			{
+				return TEXT("opencv");
+			}
+		}
+	}
+
+	return FString();
+}
+
+} // namespace UE::CaptureManager
diff --git a/Engine/Plugins/VirtualProduction/CaptureManager/CaptureManagerCore/Source/CaptureMetadataExtraction/Internal/CaptureManagerCalibrationUtils.h b/Engine/Plugins/VirtualProduction/CaptureManager/CaptureManagerCore/Source/CaptureMetadataExtraction/Internal/CaptureManagerCalibrationUtils.h
new file mode 100644
--- /dev/null
+++ b/Engine/Plugins/VirtualProduction/CaptureManager/CaptureManagerCore/Source/CaptureMetadataExtraction/Internal/CaptureManagerCalibrationUtils.h
@@ -0,0 +1,21 @@
+// Copyright Epic Games, Inc. All Rights Reserved.
+
+#pragma once
+
+#include "Containers/UnrealString.h"
+
+#define UE_API CAPTUREMETADATAEXTRACTION_API
+
+namespace UE::CaptureManager
+{
+
+// Inspects a JSON calibration file to determine whether it is OpenCV or Unreal format.
+// Returns "opencv", "unreal", or empty string if the format cannot be determined.
+// Detection mirrors the entry checks in the actual parsers:
+//   OpenCV (OpenCVCalibrationReader.cpp): root array, elements have "fx"
+//   Unreal (UnrealCalibrationParser.cpp): root object with "Calibrations" key
+UE_API FString DetectCalibrationFormat(const FString& InFilePath);
+
+} // namespace UE::CaptureManager
+
+#undef UE_API
diff --git a/Engine/Plugins/VirtualProduction/CaptureManager/CaptureManagerCore/Source/CaptureMetadataExtraction/Internal/CaptureManagerExtractionConfig.h b/Engine/Plugins/VirtualProduction/CaptureManager/CaptureManagerCore/Source/CaptureMetadataExtraction/Internal/CaptureManagerExtractionConfig.h
new file mode 100644
--- /dev/null
+++ b/Engine/Plugins/VirtualProduction/CaptureManager/CaptureManagerCore/Source/CaptureMetadataExtraction/Internal/CaptureManagerExtractionConfig.h
@@ -0,0 +1,16 @@
+// Copyright Epic Games, Inc. All Rights Reserved.
+
+#pragma once
+
+#include "Containers/UnrealString.h"
+
+namespace UE::CaptureManager
+{
+
+struct FExtractionConfig
+{
+	bool bUseFFprobe = false;
+	FString FFmpegPath;
+};
+
+} // namespace UE::CaptureManager
diff --git a/Engine/Plugins/VirtualProduction/CaptureManager/CaptureManagerCore/Source/CaptureMetadataExtraction/Internal/CaptureManagerFileExtensions.h b/Engine/Plugins/VirtualProduction/CaptureManager/CaptureManagerCore/Source/CaptureMetadataExtraction/Internal/CaptureManagerFileExtensions.h
new file mode 100644
--- /dev/null
+++ b/Engine/Plugins/VirtualProduction/CaptureManager/CaptureManagerCore/Source/CaptureMetadataExtraction/Internal/CaptureManagerFileExtensions.h
@@ -0,0 +1,40 @@
+// Copyright Epic Games, Inc. All Rights Reserved.
+
+#pragma once
+
+#include "Containers/StringView.h"
+#include "Containers/UnrealString.h"
+#include "Misc/CString.h"
+
+namespace UE::CaptureManager
+{
+
+/** Returns true if the extension (without leading dot) is a recognised video container format. Case-insensitive. */
+inline bool IsVideoExtension(FStringView Ext)
+{
+	return Ext.Equals(TEXTVIEW("mp4"), ESearchCase::IgnoreCase)
+		|| Ext.Equals(TEXTVIEW("mov"), ESearchCase::IgnoreCase)
+		|| Ext.Equals(TEXTVIEW("avi"), ESearchCase::IgnoreCase)
+		|| Ext.Equals(TEXTVIEW("mkv"), ESearchCase::IgnoreCase)
+		|| Ext.Equals(TEXTVIEW("webm"), ESearchCase::IgnoreCase);
+}
+
+/** Returns true if the extension (without leading dot) is a recognised audio format. Case-insensitive. */
+inline bool IsAudioExtension(FStringView Ext)
+{
+	return Ext.Equals(TEXTVIEW("wav"), ESearchCase::IgnoreCase)
+		|| Ext.Equals(TEXTVIEW("mp3"), ESearchCase::IgnoreCase)
+		|| Ext.Equals(TEXTVIEW("flac"), ESearchCase::IgnoreCase)
+		|| Ext.Equals(TEXTVIEW("m4a"), ESearchCase::IgnoreCase)
+		|| Ext.Equals(TEXTVIEW("aac"), ESearchCase::IgnoreCase);
+}
+
+/** Returns true if the extension (without leading dot) is a supported image-sequence format. Case-insensitive. */
+inline bool IsImageExtension(FStringView Ext)
+{
+	return Ext.Equals(TEXTVIEW("jpg"), ESearchCase::IgnoreCase)
+		|| Ext.Equals(TEXTVIEW("jpeg"), ESearchCase::IgnoreCase)
+		|| Ext.Equals(TEXTVIEW("png"), ESearchCase::IgnoreCase);
+}
+
+} // namespace UE::CaptureManager
diff --git a/Engine/Plugins/VirtualProduction/CaptureManager/CaptureManagerCore/Source/CaptureMetadataExtraction/Internal/LiveLinkFaceMetadata.cpp b/Engine/Plugins/VirtualProduction/CaptureManager/CaptureManagerCore/Source/CaptureMetadataExtraction/Internal/LiveLinkFaceMetadata.cpp
new file mode 100644
--- /dev/null
+++ b/Engine/Plugins/VirtualProduction/CaptureManager/CaptureManagerCore/Source/CaptureMetadataExtraction/Internal/LiveLinkFaceMetadata.cpp
@@ -0,0 +1,1126 @@
+// Copyright Epic Games, Inc. All Rights Reserved.
+
+#include "LiveLinkFaceMetadata.h"
+
+#include "Utils/AppleDeviceList.h"
+#include "Misc/FileHelper.h"
+#include "Serialization/JsonSerializer.h"
+
+#include "Logging/LogMacros.h"
+
+#define LOCTEXT_NAMESPACE "LiveLinkFaceMetadata"
+
+DEFINE_LOG_CATEGORY_STATIC(LogLiveLinkFaceMetadata, Log, All);
+
+namespace UE::CaptureManager 
+{
+
+enum class ELiveLinkTakeMetadataError
+{
+	InternalError,
+	InvalidArgument,
+	AbortedByUser,
+};
+
+class FLiveLinkTakeMetadataError
+{
+public:
+
+	FLiveLinkTakeMetadataError(FText InMessage,
+		ELiveLinkTakeMetadataError InCode = ELiveLinkTakeMetadataError::InternalError);
+
+	const FString GetMessageString() const;
+
+private:
+	const FText& GetMessage() const;
+	ELiveLinkTakeMetadataError GetCode() const;
+
+	FText Message;
+	ELiveLinkTakeMetadataError Code;
+};
+
+template <typename T>
+using FLiveLinkTakeMetadataResult = TValueOrError<T, FLiveLinkTakeMetadataError>;
+
+using FLiveLinkTakeMetadataVoidResult = TValueOrError<void, FLiveLinkTakeMetadataError>;
+
+struct FLiveLinkFaceStaticFileNames
+{
+	inline static const FString AudioMetadata = TEXT("audio_metadata.json");
+	inline static const FString DepthData = TEXT("depth_data.bin");
+	inline static const FString DepthMetadata = TEXT("depth_metadata.mhaical");
+	inline static const FString FrameLog = TEXT("frame_log.csv");
+	inline static const FString TakeMetadata = TEXT("take.json");
+	inline static const FString Thumbnail = TEXT("thumbnail.jpg");
+	inline static const FString VideoMetadata = TEXT("video_metadata.json");
+	inline static const FString VideoExt = TEXT(".mov");
+};
+
+struct FLiveLinkFaceTakeMetadata
+{
+	int32 Version = 0;
+	FString SlateName;
+	FString AppVersion;
+	FString DeviceModel;
+	FString DeviceClass;
+	FString Subject;
+	FString Identifier;
+	FDateTime Date;
+	int32 TakeNumber = 0;
+	int32 NumFrames = 0;
+	bool bIsCalibrated = false;
+
+	FString MOVFileName() const;
+
+private:
+	FString CommonFileNamePrefix() const;
+	TArray<FString> GetCalibratedBlendshapeFileNames() const;
+	TArray<FString> GetMHAFileNames() const;
+	TArray<FString> GetCommonFileNames() const;
+	TArray<FString> GetARKitFileNames() const;
+};
+
+struct FLiveLinkFaceVideoMetadata
+{
+	FIntPoint Resolution;
+	float JpegCompressionQuality;
+	FString Quality;
+	EMediaOrientation Orientation;
+	double FrameRate;
+};
+
+struct FLiveLinkFaceOodleMetadata
+{
+	FString Compressor;
+	FString CompressionLevel;
+	FString Version;
+};
+
+struct FLiveLinkFaceDepthMetadata
+{
+	FString Build;
+	FIntPoint Resolution;
+	FString Compression;
+	FString DeviceModel;
+	FString DeviceClass;
+	FLiveLinkFaceOodleMetadata OodleInfo; // If Compression == "Oodle"
+	EMediaOrientation Orientation;
+	double FrameRate;
+	float PixelSize;
+	bool bShouldCompressFiles = false;
+
+	// Lens distortion info
+	TArray<float> LensDistortionLookupTable;
+	TArray<float> InverseLensDistortionLookupTable;
+	TArray<float> IntrinsicMatrix;
+	FVector2D LensDistortionCenter;
+	FVector2D IntrinsicMatrixReferenceDimensions;
+};
+
+struct FLiveLinkFaceAudioMetadata
+{
+	int32 BitsPerChannel = 0;
+	int32 SampleRate = 0;
+	int32 ChannelsPerFrame = 0;
+	int32 FormatFlags = 0;
+};
+
+struct FLiveLinkFaceTakeInfo
+{
+	// The path to the folder where the files for this take are
+	FString TakeOriginDirectory;
+
+	TArray<uint8> RawThumbnailData;
+	FLiveLinkFaceTakeMetadata TakeMetadata;
+	FLiveLinkFaceVideoMetadata VideoMetadata;
+	FLiveLinkFaceDepthMetadata DepthMetadata;
+	FLiveLinkFaceAudioMetadata AudioMetadata;
+
+	// Set by TryResolveVideoFilePath when the metadata-derived filename does not match the actual file on disk.
+	TOptional<FString> VideoFilePathOverride;
+
+	FString GetTakeName() const;
+	FString GetVideoFilePath() const;
+	FString GetDepthFilePath() const;
+	FString GetFrameLogFilePath() const;
+	FString GetCameraCalibrationFilePath() const;
+
+private:
+	FString GetTakePath() const;
+	FString GetOutputDirectory(const FString& InTakesOriginDirectory) const;
+};
+
+class FLiveLinkFaceMetadataParser
+{
+public:
+
+	static EMediaOrientation ParseOrientation(int32 InOrientation = 4);
+	static bool ParseVideoMetadata(const FString& InTakeDirectory, FLiveLinkFaceVideoMetadata& OutVideoMetadata);
+	static bool ParseVideoMetadataFromString(const FString& InJsonString, FLiveLinkFaceVideoMetadata& OutVideoMetadata);
+	static bool ParseDepthMetadata(const FString& InTakeDirectory, FLiveLinkFaceDepthMetadata& OutDepthMetadata);
+	static bool ParseAudioMetadata(const FString& InTakeDirectory, FLiveLinkFaceAudioMetadata& OutAudioMetadata);
+	static bool ParseTakeInfo(const FString& InTakeDirectory, FLiveLinkFaceTakeInfo& OutTakeInfo);
+	static bool ParseThumbnail(const FString& InTakeDirectory, FLiveLinkFaceTakeInfo& OutTakeInfo);
+	static void ParseiOSDeviceModel(const FString& InDeviceModel, FString& OutiOSDeviceClass);
+
+private:
+
+	static FLiveLinkTakeMetadataResult<FLiveLinkFaceTakeMetadata> ParseTakeMetadata(const TSharedPtr<FJsonObject> JsonObject);
+	static FLiveLinkTakeMetadataVoidResult ParseString(const TSharedPtr<FJsonObject> JsonObject, const FString& Key, FString& OutString);
+	static FLiveLinkTakeMetadataVoidResult ParseNumber(const TSharedPtr<FJsonObject> JsonObject, const FString& Key, int32& OutNumber);
+	static FLiveLinkTakeMetadataVoidResult ParseBool(const TSharedPtr<FJsonObject> JsonObject, const FString& Key, bool& OutBool);
+	static FLiveLinkTakeMetadataError CreateErrorForMissingJsonKey(const FString& Key);
+	static TArray<TSharedPtr<class FJsonValue>> ParseJsonArrayFromFile(const FString& InFilePath);
+	static const FString ParseJsonStringFromFile(const FString& InFilePath);
+	static TSharedPtr<class FJsonObject> ParseJsonObjectFromString(const FString& InJsonString);
+};
+
+static bool ContainsWhitespace(const FString& InStr)
+{
+	for (const FString::ElementType& Character : InStr)
+	{
+		if (FChar::IsWhitespace(Character))
+		{
+			return true;
+		}
+	}
+	return false;
+}
+
+FLiveLinkTakeMetadataError::FLiveLinkTakeMetadataError(FText InMessage,
+													   ELiveLinkTakeMetadataError InCode)
+	: Message(MoveTemp(InMessage))
+	, Code(InCode)
+{
+}
+
+const FText& FLiveLinkTakeMetadataError::GetMessage() const
+{
+	return Message;
+}
+
+const FString FLiveLinkTakeMetadataError::GetMessageString() const
+{
+	return Message.ToString();
+}
+
+ELiveLinkTakeMetadataError FLiveLinkTakeMetadataError::GetCode() const
+{
+	return Code;
+}
+
+TArray<FString> FLiveLinkFaceTakeMetadata::GetCommonFileNames() const
+{
+	return {
+		FLiveLinkFaceStaticFileNames::AudioMetadata,
+		FLiveLinkFaceStaticFileNames::FrameLog,
+		FLiveLinkFaceStaticFileNames::TakeMetadata,
+		FLiveLinkFaceStaticFileNames::Thumbnail,
+		FLiveLinkFaceStaticFileNames::VideoMetadata,
+		MOVFileName()
+	};
+}
+
+TArray<FString> FLiveLinkFaceTakeMetadata::GetMHAFileNames() const
+{
+	TArray FileNames = GetCommonFileNames();
+	FileNames.Append({
+		FLiveLinkFaceStaticFileNames::DepthData,
+		FLiveLinkFaceStaticFileNames::DepthMetadata
+	});
+	return FileNames;
+}
+
+TArray<FString> FLiveLinkFaceTakeMetadata::GetARKitFileNames() const
+{
+	TArray FileNames = GetCommonFileNames();
+	if (bIsCalibrated)
+	{
+		FileNames.Append(GetCalibratedBlendshapeFileNames());
+	}
+	else
+	{
+		const FString BlendshapeFileName = FString::Format(TEXT("{0}.csv"), { CommonFileNamePrefix() });
+		FileNames.Add(BlendshapeFileName);
+	}
+	return FileNames;
+}
+
+FString FLiveLinkFaceTakeMetadata::CommonFileNamePrefix() const
+{
+	return FString::Format(TEXT("{0}_{1}_{2}"), { SlateName, TakeNumber, Subject });
+}
+
+TArray<FString> FLiveLinkFaceTakeMetadata::GetCalibratedBlendshapeFileNames() const
+{
+	FString Prefix = CommonFileNamePrefix();
+	return {
+		FString::Format(TEXT("{0}_cal.csv"), { Prefix }),
+		FString::Format(TEXT("{0}_neutral.csv"), { Prefix }),
+		FString::Format(TEXT("{0}_raw.csv"), { Prefix })
+	};
+}
+
+FString FLiveLinkFaceTakeMetadata::MOVFileName() const
+{
+	return FString::Format(TEXT("{0}.mov"), { CommonFileNamePrefix() });
+}
+
+FString FLiveLinkFaceTakeInfo::GetTakeName() const
+{	
+	return FString::Format(TEXT("{0}_{1}"), { TakeMetadata.SlateName, TakeMetadata.TakeNumber });
+}
+
+FString FLiveLinkFaceTakeInfo::GetTakePath() const
+{
+	return FString::Format(TEXT("{0}_{1}_{2}"), { TakeMetadata.SlateName, TakeMetadata.TakeNumber, TakeMetadata.Subject });
+}
+
+FString FLiveLinkFaceTakeInfo::GetVideoFilePath() const
+{
+	if (VideoFilePathOverride.IsSet())
+	{
+		return VideoFilePathOverride.GetValue();
+	}
+	return FString::Format(TEXT("{0}/{1}"), { TakeOriginDirectory, TakeMetadata.MOVFileName() });
+}
+
+FString FLiveLinkFaceTakeInfo::GetDepthFilePath() const
+{
+	return FString::Format(TEXT("{0}/{1}"), { TakeOriginDirectory, FLiveLinkFaceStaticFileNames::DepthData });
+}
+
+FString FLiveLinkFaceTakeInfo::GetFrameLogFilePath() const
+{
+	return FString::Format(TEXT("{0}/{1}"), { TakeOriginDirectory, FLiveLinkFaceStaticFileNames::FrameLog });
+}
+
+FString FLiveLinkFaceTakeInfo::GetCameraCalibrationFilePath() const
+{
+	return FString::Format(TEXT("{0}/{1}"), { TakeOriginDirectory, FLiveLinkFaceStaticFileNames::DepthMetadata });
+}
+
+FString FLiveLinkFaceTakeInfo::GetOutputDirectory(const FString& InTakesOriginDirectory) const
+{
+	return TakeOriginDirectory.Mid(InTakesOriginDirectory.Len());
+}
+
+EMediaOrientation FLiveLinkFaceMetadataParser::ParseOrientation(int32 InOrientation)
+{
+	// 0: Portrait, 1: Landscape, 2: Portrait, 3: Landscape
+	switch (InOrientation)
+	{
+		case 1: // Portrait
+			return EMediaOrientation::Original;
+		case 2:	// PortraitUpsideDown
+			return EMediaOrientation::CW180;
+		case 3: // landscapeLeft
+			return EMediaOrientation::CW90;
+		case 4: // LandscapeRight
+		default:
+			return EMediaOrientation::CW270;
+	}
+}
+
+bool FLiveLinkFaceMetadataParser::ParseVideoMetadata(const FString& InTakeDirectory, FLiveLinkFaceVideoMetadata& OutVideoMetadata)
+{
+	const FString& VideoMetadataFile = InTakeDirectory / FLiveLinkFaceStaticFileNames::VideoMetadata;
+	const FString& VideoMetadataString = ParseJsonStringFromFile(VideoMetadataFile);
+
+	return !VideoMetadataString.IsEmpty() && ParseVideoMetadataFromString(VideoMetadataString, OutVideoMetadata);
+}
+
+bool FLiveLinkFaceMetadataParser::ParseVideoMetadataFromString(const FString& InJsonString, FLiveLinkFaceVideoMetadata& OutVideoMetadata)
+{
+	if (TSharedPtr<FJsonObject> VideoMetadataJson = ParseJsonObjectFromString(InJsonString))
+	{
+		// Sample video_metadata.json
+		// {
+		// 	"FrameRate" : 60,
+		// 	"Dimensions" : {
+		// 		"width" : 1280,
+		// 		"height" : 720
+		// 	},
+		// 	"Quality" : "high",
+		// 	"JpegCompressionQuality" : 0.90000000000000002
+		// }
+
+		VideoMetadataJson->TryGetNumberField(TEXT("FrameRate"), OutVideoMetadata.FrameRate);
+
+		const TSharedPtr<FJsonObject>* DimensionsJson;
+		if (VideoMetadataJson->TryGetObjectField(TEXT("Dimensions"), DimensionsJson))
+		{
+			DimensionsJson->ToSharedRef()->TryGetNumberField(TEXT("width"), OutVideoMetadata.Resolution.X);
+			DimensionsJson->ToSharedRef()->TryGetNumberField(TEXT("height"), OutVideoMetadata.Resolution.Y);
+		}
+
+		VideoMetadataJson->TryGetStringField(TEXT("Quality"), OutVideoMetadata.Quality);
+		VideoMetadataJson->TryGetNumberField(TEXT("JpegCompressionQuality"), OutVideoMetadata.JpegCompressionQuality);
+
+		int32 Orientation = 4;
+		VideoMetadataJson->TryGetNumberField(TEXT("Orientation"), Orientation);
+		OutVideoMetadata.Orientation = ParseOrientation(Orientation);
+
+		return true;
+	}
+
+	return false;
+}
+
+bool FLiveLinkFaceMetadataParser::ParseDepthMetadata(const FString& InTakeDirectory, FLiveLinkFaceDepthMetadata& OutDepthMetadata)
+{
+	const FString& DepthMetadataFile = InTakeDirectory / FLiveLinkFaceStaticFileNames::DepthMetadata;
+
+	const FString& DepthMetadataString = ParseJsonStringFromFile(DepthMetadataFile);
+
+	if (TSharedPtr<FJsonObject> DepthMetadataJson = ParseJsonObjectFromString(DepthMetadataString))
+	{
+		DepthMetadataJson->TryGetStringField(TEXT("Build"), OutDepthMetadata.Build);
+		DepthMetadataJson->TryGetStringField(TEXT("Compression"), OutDepthMetadata.Compression);
+		DepthMetadataJson->TryGetNumberField(TEXT("PixelSize"), OutDepthMetadata.PixelSize);
+		DepthMetadataJson->TryGetNumberField(TEXT("DepthFrameRate"), OutDepthMetadata.FrameRate);
+		DepthMetadataJson->TryGetStringField(TEXT("DeviceModel"), OutDepthMetadata.DeviceModel);
+		ParseiOSDeviceModel(OutDepthMetadata.DeviceModel, OutDepthMetadata.DeviceClass);
+
+		int32 Orientation = 4;
+		DepthMetadataJson->TryGetNumberField(TEXT("Orientation"), Orientation);
+		OutDepthMetadata.Orientation = ParseOrientation(Orientation);
+
+		if (OutDepthMetadata.Compression == TEXT("Oodle"))
+		{
+			const TSharedPtr<FJsonObject>* OodleCompression;
+			if (DepthMetadataJson->TryGetObjectField(TEXT("Oodle"), OodleCompression))
+			{
+				OodleCompression->ToSharedRef()->TryGetStringField(TEXT("CompressionLevel"), OutDepthMetadata.OodleInfo.CompressionLevel);
+				OodleCompression->ToSharedRef()->TryGetStringField(TEXT("Compressor"), OutDepthMetadata.OodleInfo.Compressor);
+				OodleCompression->ToSharedRef()->TryGetStringField(TEXT("Version"), OutDepthMetadata.OodleInfo.Version);
+			}
+		}
+
+		const TArray<TSharedPtr<FJsonValue>>* LensDistortionLUTArray;
+		if (DepthMetadataJson->TryGetArrayField(TEXT("LensDistortionLookupTable"), LensDistortionLUTArray))
+		{
+			for (const TSharedPtr<FJsonValue>& Value : *LensDistortionLUTArray)
+			{
+				OutDepthMetadata.LensDistortionLookupTable.Add(Value->AsNumber());
+			}
+		}
+
+		const TArray<TSharedPtr<FJsonValue>>* InverseLensDistortionLUTArray;
+		if (DepthMetadataJson->TryGetArrayField(TEXT("InverseLensDistortionLookupTable"), InverseLensDistortionLUTArray))
+		{
+			for (const TSharedPtr<FJsonValue>& Value : *InverseLensDistortionLUTArray)
+			{
+				OutDepthMetadata.InverseLensDistortionLookupTable.Add(Value->AsNumber());
+			}
+		}
+
+		const TSharedPtr<FJsonObject>* IntrinsicMatrixReferenceDimensions;
+		if (DepthMetadataJson->TryGetObjectField(TEXT("IntrinsicMatrixReferenceDimensions"), IntrinsicMatrixReferenceDimensions))
+		{
+			IntrinsicMatrixReferenceDimensions->ToSharedRef()->TryGetNumberField(TEXT("Width"), OutDepthMetadata.IntrinsicMatrixReferenceDimensions.X);
+			IntrinsicMatrixReferenceDimensions->ToSharedRef()->TryGetNumberField(TEXT("Height"), OutDepthMetadata.IntrinsicMatrixReferenceDimensions.Y);
+		}
+
+		const TSharedPtr<FJsonObject>* DepthDimensions;
+		if (DepthMetadataJson->TryGetObjectField(TEXT("DepthDimensions"), DepthDimensions))
+		{
+			DepthDimensions->ToSharedRef()->TryGetNumberField(TEXT("Width"), OutDepthMetadata.Resolution.X);
+			DepthDimensions->ToSharedRef()->TryGetNumberField(TEXT("Height"), OutDepthMetadata.Resolution.Y);
+		}
+
+		const TSharedPtr<FJsonObject>* LensDistorionCenter;
+		if (DepthMetadataJson->TryGetObjectField(TEXT("LensDistortionCenter"), LensDistorionCenter))
+		{
+			LensDistorionCenter->ToSharedRef()->TryGetNumberField(TEXT("X"), OutDepthMetadata.LensDistortionCenter.X);
+			LensDistorionCenter->ToSharedRef()->TryGetNumberField(TEXT("Y"), OutDepthMetadata.LensDistortionCenter.Y);
+		}
+
+		const TArray<TSharedPtr<FJsonValue>>* IntrinsicMatrix;
+		if (DepthMetadataJson->TryGetArrayField(TEXT("IntrinsicMatrix"), IntrinsicMatrix))
+		{
+			for (const TSharedPtr<FJsonValue>& Value : *IntrinsicMatrix)
+			{
+				OutDepthMetadata.IntrinsicMatrix.Add(Value->AsNumber());
+			}
+		}
+
+		return true;
+	}
+
+	return false;
+}
+
+bool FLiveLinkFaceMetadataParser::ParseAudioMetadata(const FString& InTakeDirectory, FLiveLinkFaceAudioMetadata& OutAudioMetadata)
+{
+	const FString& AudioMetadataFile = InTakeDirectory / FLiveLinkFaceStaticFileNames::AudioMetadata;
+
+	const FString& AudioMetadataString = ParseJsonStringFromFile(AudioMetadataFile);
+
+	if (TSharedPtr<FJsonObject> AudioMetadataJson = ParseJsonObjectFromString(AudioMetadataString))
+	{
+		// Sample audio_metadata.json
+		// {
+		// 	"BitsPerChannel" : 16,
+		// 	"SampleRate" : 44100,
+		// 	"ChannelsPerFrame" : 1,
+		// 	"FormatFlags" : 12
+		// }
+
+		AudioMetadataJson->TryGetNumberField(TEXT("BitsPerChannel"), OutAudioMetadata.BitsPerChannel);
+		AudioMetadataJson->TryGetNumberField(TEXT("SampleRate"), OutAudioMetadata.SampleRate);
+		AudioMetadataJson->TryGetNumberField(TEXT("ChannelsPerFrame"), OutAudioMetadata.ChannelsPerFrame);
+		AudioMetadataJson->TryGetNumberField(TEXT("FormatFlags"), OutAudioMetadata.FormatFlags);
+
+		return true;
+	}
+
+	return false;
+}
+
+bool FLiveLinkFaceMetadataParser::ParseTakeInfo(const FString& InTakeDirectory, FLiveLinkFaceTakeInfo& OutTakeInfo)
+{
+	const FString TakeMetadataFile = InTakeDirectory / FLiveLinkFaceStaticFileNames::TakeMetadata;
+
+	const FString& TakeMetadataString = ParseJsonStringFromFile(TakeMetadataFile);
+	const TSharedPtr TakeMetadataJson = ParseJsonObjectFromString(TakeMetadataString);
+
+	if (!TakeMetadataJson.IsValid())
+	{
+		return false;
+	}
+	
+	// Sample take.json
+	// {
+	// 	"frames" : 1058,
+	// 	"appVersion" : "v0.1.0 (build 17)",
+	// 	"slate" : "HH_Neutral_Rotation",
+	// 	"calibrated" : false,
+	// 	"subject" : "iPhone12Pro",
+	// 	"identifier" : "20211006_HH_Neutral_Rotation_1",
+	// 	"date" : "2021-10-06T11:31:46Z",
+	// 	"version" : 1,
+	// 	"take" : 1,
+	// 	"deviceModel" : "iPhone13,3"
+	// }
+
+	FLiveLinkTakeMetadataResult<FLiveLinkFaceTakeMetadata> ParseResult = ParseTakeMetadata(TakeMetadataJson);
+
+	if (ParseResult.HasError())
+	{
+		const FLiveLinkTakeMetadataError& Error = ParseResult.StealError();
+		UE_LOGF(LogLiveLinkFaceMetadata, Warning, "Failed to parse take metadata: %ls", *Error.GetMessageString());
+		return false;
+	}
+
+	if (!ParseResult.HasValue())
+	{
+		UE_LOGF(LogLiveLinkFaceMetadata, Error, "Take metadata parse result did not contain an error but no value was provided");
+		return false;
+	}
+
+	OutTakeInfo.TakeMetadata = ParseResult.StealValue();
+	OutTakeInfo.TakeOriginDirectory = InTakeDirectory;
+	
+	return true;
+}
+
+bool FLiveLinkFaceMetadataParser::ParseThumbnail(const FString& InTakeDirectory, FLiveLinkFaceTakeInfo& OutTakeInfo)
+{
+	const FString& ThumbnailFile = InTakeDirectory / FLiveLinkFaceStaticFileNames::Thumbnail;
+
+	return FFileHelper::LoadFileToArray(OutTakeInfo.RawThumbnailData, *ThumbnailFile);
+}
+
+void FLiveLinkFaceMetadataParser::ParseiOSDeviceModel(const FString& InDeviceModel, FString& OutiOSDeviceClass)
+{
+	if (InDeviceModel.IsEmpty())
+	{
+		OutiOSDeviceClass = TEXT("Unspecified");
+		return;
+	}
+
+	// Device model is in format [Name][FirstNumber],[SecondNumber]
+	TArray<FString> ParsedStringArray;
+	InDeviceModel.ParseIntoArray(ParsedStringArray, TEXT(","));
+
+	if (ParsedStringArray[0].Contains(TEXT("iPhone")))
+	{
+		if (FAppleDeviceList::DeviceMap.Contains(InDeviceModel))
+		{
+			OutiOSDeviceClass = FAppleDeviceList::DeviceMap[InDeviceModel];
+		}
+		else
+		{
+			OutiOSDeviceClass = TEXT("iPhone");
+		}
+	}
+	else if (ParsedStringArray[0].Contains(TEXT("iPad")))
+	{
+		OutiOSDeviceClass = TEXT("iPad");
+	}
+	else
+	{
+		OutiOSDeviceClass = TEXT("Unspecified");
+	}
+}
+
+FLiveLinkTakeMetadataResult<FLiveLinkFaceTakeMetadata> FLiveLinkFaceMetadataParser::ParseTakeMetadata(const TSharedPtr<FJsonObject> JsonObject)
+{
+	FLiveLinkFaceTakeMetadata TakeMetadata;
+
+	FString DateString;
+	const TArray ParseErrors = {
+		ParseString(JsonObject, TEXT("slate"), TakeMetadata.SlateName),
+		ParseString(JsonObject, TEXT("appVersion"), TakeMetadata.AppVersion),
+		ParseString(JsonObject, TEXT("deviceModel"), TakeMetadata.DeviceModel),
+		ParseString(JsonObject, TEXT("subject"), TakeMetadata.Subject),
+		ParseString(JsonObject, TEXT("identifier"), TakeMetadata.Identifier),
+		ParseString(JsonObject, TEXT("date"), DateString),
+		ParseNumber(JsonObject, TEXT("version"), TakeMetadata.Version),
+		ParseNumber(JsonObject, TEXT("take"), TakeMetadata.TakeNumber),
+		ParseNumber(JsonObject, TEXT("frames"), TakeMetadata.NumFrames),
+		ParseBool(JsonObject, TEXT("calibrated"), TakeMetadata.bIsCalibrated)
+	};
+
+	TArray Errors = ParseErrors.FilterByPredicate([](const FLiveLinkTakeMetadataVoidResult& Error) { return Error.HasError(); });
+
+	if (!Errors.IsEmpty())
+	{			
+		return MakeError(Errors[0].StealError());
+	}
+	
+	ParseiOSDeviceModel(TakeMetadata.DeviceModel, TakeMetadata.DeviceClass);
+	
+	if (!FDateTime::ParseIso8601(*DateString, TakeMetadata.Date))
+	{
+		return MakeError(FLiveLinkTakeMetadataError(FText::Format(LOCTEXT("LiveLinkFaceMetadata_DateError", "Failed to parse DateString '{0}' as an Iso8601 date"), FText::FromString(DateString))));
+	}
+
+	return MakeValue(TakeMetadata);
+}
+
+FLiveLinkTakeMetadataVoidResult FLiveLinkFaceMetadataParser::ParseString(const TSharedPtr<FJsonObject> JsonObject, const FString& Key, FString &OutString)
+{
+	if (!JsonObject->TryGetStringField(Key, OutString))
+	{
+		return MakeError(CreateErrorForMissingJsonKey(Key));
+	}
+
+	return MakeValue();
+}
+
+FLiveLinkTakeMetadataVoidResult FLiveLinkFaceMetadataParser::ParseNumber(const TSharedPtr<FJsonObject> JsonObject, const FString& Key, int32 &OutNumber)
+{
+	if (!JsonObject->TryGetNumberField(Key, OutNumber))
+	{
+		return MakeError(CreateErrorForMissingJsonKey(Key));
+	}
+
+	return MakeValue();
+}
+
+FLiveLinkTakeMetadataVoidResult FLiveLinkFaceMetadataParser::ParseBool(const TSharedPtr<FJsonObject> JsonObject, const FString& Key, bool &OutBool)
+{
+	if (!JsonObject->TryGetBoolField(Key, OutBool))
+	{
+		return MakeError(CreateErrorForMissingJsonKey(Key));
+	}
+
+	return MakeValue();
+}
+
+FLiveLinkTakeMetadataError FLiveLinkFaceMetadataParser::CreateErrorForMissingJsonKey(const FString& Key)
+{
+	return FLiveLinkTakeMetadataError(FText::Format(LOCTEXT("LiveLinkFaceMetadata_MissingValue", "Missing value for key '{0}' in Live Link Face Metadata"), FText::FromString(Key)));
+}
+
+TArray<TSharedPtr<FJsonValue>> FLiveLinkFaceMetadataParser::ParseJsonArrayFromFile(const FString& InFilePath)
+{
+	FString JsonStringBuffer;
+	TArray<TSharedPtr<FJsonValue>> Result;
+
+	if (FFileHelper::LoadFileToString(JsonStringBuffer, *InFilePath) &&
+		FJsonSerializer::Deserialize(TJsonReaderFactory<>::Create(JsonStringBuffer), Result))
+	{
+		return Result;
+	}
+
+	return TArray<TSharedPtr<FJsonValue>>();
+}
+
+const FString FLiveLinkFaceMetadataParser::ParseJsonStringFromFile(const FString& InFilePath)
+{
+	FString JsonStringBuffer;
+
+	if (!IFileManager::Get().FileExists(*InFilePath))
+	{
+		UE_LOGF(LogLiveLinkFaceMetadata, Error, "File not found: %ls", *InFilePath);
+		return {};
+	}
+
+	if (!FFileHelper::LoadFileToString(JsonStringBuffer, *InFilePath))
+	{
+		UE_LOGF(LogLiveLinkFaceMetadata, Error, "Failed to load file (check permissions): %ls", *InFilePath);
+		return {};
+	}
+
+	return JsonStringBuffer;
+}
+
+TSharedPtr<FJsonObject> FLiveLinkFaceMetadataParser::ParseJsonObjectFromString(const FString& InJsonString)
+{
+	TSharedPtr<FJsonObject> Result;
+
+	if (!FJsonSerializer::Deserialize(TJsonReaderFactory<>::Create(InJsonString), Result))
+	{
+		UE_LOGF(LogLiveLinkFaceMetadata, Error, "Failed to parse json string (check for syntax errors): %ls", *InJsonString);
+		return nullptr;
+	}
+
+	return Result;
+}
+
+class FrameLogEntry
+{
+public:
+	static const char VideoType = 'V';
+	static const char DepthType = 'D';
+	static const char AudioType = 'A';
+	static const char InvalidType = '\0';
+
+	FrameLogEntry() = default;
+
+	static bool Parse(const FString& InLogLine, FrameLogEntry& OutLogEntry)
+	{
+		TArray<FString> Tokens;
+		InLogLine.ParseIntoArray(Tokens, TEXT(","));
+
+		if (Tokens.Num() < 5 || Tokens.Num() > 6)
+		{
+			return false;
+		}
+
+		OutLogEntry = FrameLogEntry{ MoveTemp(Tokens) };
+		return true;
+	}
+
+	char EntryType()
+	{
+		return !Tokens[0].IsEmpty() ? Tokens[0][0] : InvalidType;
+	}
+
+	int64 FrameIndex()
+	{
+		return FCString::Atoi64(*Tokens[1]);
+	}
+
+	double Time()
+	{
+		const int64 Numerator = FCString::Atoi64(*Tokens[2]);
+		const double Denominator = FCString::Atod(*Tokens[3]);
+		return Numerator / Denominator;
+	}
+
+	bool Timecode(FTimecode& OutTimecode)
+	{
+		TArray<FString> TimecodeTokens;
+		Tokens[4].ParseIntoArray(TimecodeTokens, TEXT(":"));
+		if (TimecodeTokens.Num() != 4 && TimecodeTokens.Num() != 3)
+		{
+			return false;
+		}
+
+		// Limit hours to 0-23 else we can't accurately
+		// show the clip in sequencer.
+		const int32 Hours = FCString::Atoi(*TimecodeTokens[0]) % 24;
+		const int32 Mins = FCString::Atoi(*TimecodeTokens[1]);
+
+		int32 Secs = 0;
+		int32 Frames = 0;
+		bool bIsDropFrame = TimecodeTokens[2].Contains(TEXT(";"));
+		if (bIsDropFrame)
+		{
+			// TimecodeTokens[2] == 00;00
+			TArray<FString> SecondsAndFrames;
+			TimecodeTokens[2].ParseIntoArray(SecondsAndFrames, TEXT(";"));
+
+			if (SecondsAndFrames.Num() != 2)
+			{
+				return false;
+			}
+
+			Secs = FCString::Atoi(*SecondsAndFrames[0]);
+			Frames = FMath::RoundHalfFromZero(FCString::Atof(*SecondsAndFrames[1]));
+		}
+		else
+		{
+			Secs = FCString::Atoi(*TimecodeTokens[2]);
+			Frames = FMath::RoundHalfFromZero(FCString::Atof(*TimecodeTokens[3]));
+		}
+
+		// iPhone timecode is never drop frame - always either 30 or 60 fps
+		OutTimecode = FTimecode(Hours, Mins, Secs, Frames, bIsDropFrame);
+		return true;
+	}
+
+	bool IsDroppedFrame()
+	{
+		// If it's the old log format that didn't include dropped frame info, assume it wasn't dopped
+		if (Tokens.Num() == 5)
+		{
+			return false;
+		}
+
+		return Tokens[5] != TEXT("0");
+	}
+
+private:
+
+	FrameLogEntry(TArray<FString>&& InTokens) : Tokens(MoveTemp(InTokens))
+	{
+	}
+
+	TArray<FString> Tokens;
+};
+
+class FFrameLogParser
+{
+public:
+
+	void ParseFrameLog(const FString& InFrameLogPath, float InVideoFrameRate, float InDepthFrameRate)
+	{
+		TArray<FString> FrameLogLines;
+		FFileHelper::LoadFileToStringArray(FrameLogLines, *InFrameLogPath);
+
+		bool bAudioTimecodeFound = false;
+
+		FFrameRate OriginalVideoFrameRate;
+		uint32 RoundedVideoFrameRate = static_cast<uint32>(FMath::RoundToInt(InVideoFrameRate));
+		if (RoundedVideoFrameRate != 0)
+		{
+			OriginalVideoFrameRate = FFrameRate{ RoundedVideoFrameRate, 1 };
+		}
+
+		for (const FString& Line : FrameLogLines)
+		{
+			FrameLogEntry LogEntry;
+			if (!FrameLogEntry::Parse(Line, LogEntry))
+			{
+				continue;
+			}
+
+			if (LogEntry.EntryType() == FrameLogEntry::VideoType)
+			{
+				if (LogEntry.FrameIndex() == 0) // Take timecode from first video frame
+				{
+					FTimecode Timecode;
+					if (!LogEntry.Timecode(Timecode))
+					{
+						continue;
+					}
+
+					// If frame rate isn't known, there's no sense to parse timecode
+					if (OriginalVideoFrameRate == FFrameRate{})
+					{
+						continue;
+					}
+
+					// We re-construct the timecode using the frame rate to resolve any invalid frame number rounding
+					// which may have occurred while parsing the frame log. As an example: a fractional frame number of
+					// 59.780 would round to 60, which is an invalid frame number at 60FPS. Going to/from a frame number
+					// fixes the problem.
+					VideoTimecode = FTimecode::FromFrameNumber(Timecode.ToFrameNumber(OriginalVideoFrameRate), OriginalVideoFrameRate);
+					VideoTimecodeRate = OriginalVideoFrameRate;
+				}
+			}
+			else if (LogEntry.EntryType() == FrameLogEntry::DepthType)
+			{
+				if (LogEntry.FrameIndex() == 0) // Take timecode from first depth frame
+				{
+					FTimecode Timecode;
+					if (!LogEntry.Timecode(Timecode))
+					{
+						continue;
+					}
+
+					uint32 RoundedDepthFrameRate = static_cast<uint32>(FMath::RoundToInt(InDepthFrameRate));
+					if (RoundedDepthFrameRate == 0)
+					{
+						continue;
+					}
+
+					FFrameRate OriginalDepthFrameRate = FFrameRate{ RoundedDepthFrameRate, 1 };
+
+					// This step is also needed to resolve any invalid frame number rounding which may have occurred
+					// during the parsing of the frame log (as it implicitly does a conversion to/from frame number).
+					// The depth timecode in the frame log matches the frame rate of the video frame rate so we use the video frame rate when creating the timespan
+					DepthTimecode = FTimecode::FromTimespan(Timecode.ToTimespan(OriginalVideoFrameRate), OriginalDepthFrameRate, true);
+					DepthTimecodeRate = OriginalDepthFrameRate;
+				}
+
+				++DepthFrameCount;
+			}
+			else if (LogEntry.EntryType() == FrameLogEntry::AudioType)
+			{
+				if (!bAudioTimecodeFound)
+				{
+					FTimecode Timecode;
+					if (!LogEntry.Timecode(Timecode))
+					{
+						continue;
+					}
+
+					// If frame rate isn't known, there's no sense to parse timecode
+					if (OriginalVideoFrameRate == FFrameRate{})
+					{
+						continue;
+					}
+
+					// Originally we intentionally used the same frame rate here as for depth (assuming depth at 30 FPS), but now depth frame rate can vary.
+					// In the near future we could consider matching the video frame rate here instead, but for the sake of
+					// reducing risk for 5.7 we will maintain the current behaviour.
+					FFrameRate TargetFrameRate{ 30, 1 };
+
+					// The audio timecode is expressed at the video frame rate 
+					// This step is also needed to resolve any invalid frame number rounding which may have occurred
+					// during the parsing of the frame log (as it implicitly does a conversion to/from frame number).
+					AudioTimecode = FTimecode::FromTimespan(Timecode.ToTimespan(OriginalVideoFrameRate), TargetFrameRate, true);
+					AudioTimecodeRate = TargetFrameRate;
+					bAudioTimecodeFound = true;
+				}
+			}
+		}
+
+		// If no audio timecode was specified, assume it's the same as video
+		if (!bAudioTimecodeFound)
+		{
+			AudioTimecode = VideoTimecode;
+		}
+	}
+
+	FTimecode VideoTimecode;
+	FTimecode DepthTimecode;
+	FTimecode AudioTimecode;
+
+	FFrameRate VideoTimecodeRate;
+	FFrameRate DepthTimecodeRate;
+	FFrameRate AudioTimecodeRate;
+
+	int32 DepthFrameCount = 0;
+};
+
+static FTakeMetadata::FVideo::EOrientation ConvertOldOrientation(EMediaOrientation InOrientation)
+{
+	switch (InOrientation)
+	{
+		case EMediaOrientation::Original:
+			return FTakeMetadata::FVideo::EOrientation::Original;
+		case EMediaOrientation::CW90:
+			return FTakeMetadata::FVideo::EOrientation::CW90;
+		case EMediaOrientation::CW180:
+			return FTakeMetadata::FVideo::EOrientation::CW180;
+		case EMediaOrientation::CW270:
+		default:
+			return FTakeMetadata::FVideo::EOrientation::CW270;
+	}
+}
+
+bool HasDepthMetadataFile(const FString& InTakeDirectory)
+{
+	const FString& DepthMetadataFile = InTakeDirectory / FLiveLinkFaceStaticFileNames::DepthMetadata;
+
+	return IFileManager::Get().FileExists(*DepthMetadataFile);
+}
+
+FTakeMetadata ConvertOldToNewTakeMetadata(const FLiveLinkFaceTakeInfo& InTakeInfo)
+{
+	FTakeMetadata NewTakeInfo;
+
+	NewTakeInfo.Version.Major = 3;
+	NewTakeInfo.Version.Minor = 0;
+
+	// Create new guid for old takes
+	NewTakeInfo.UniqueId = FGuid::NewGuid().ToString(EGuidFormats::DigitsWithHyphensLower);
+	NewTakeInfo.DateTime = InTakeInfo.TakeMetadata.Date;
+	NewTakeInfo.TakeNumber = InTakeInfo.TakeMetadata.TakeNumber;
+	NewTakeInfo.Slate = InTakeInfo.TakeMetadata.SlateName;
+
+	NewTakeInfo.Thumbnail = FTakeThumbnailData(InTakeInfo.TakeOriginDirectory / FLiveLinkFaceStaticFileNames::Thumbnail);
+
+	NewTakeInfo.Device.Model = InTakeInfo.TakeMetadata.DeviceModel;
+
+	FTakeMetadata::FDevice::FSoftware Software;
+	Software.Version = InTakeInfo.TakeMetadata.AppVersion;
+	NewTakeInfo.Device.Software.Add(MoveTemp(Software));
+
+	FTakeMetadata::FDevice::FPlatform Platform;
+	Platform.Name = TEXT("iOS");
+	Platform.Version = FString();
+	NewTakeInfo.Device.Platform = MoveTemp(Platform);
+
+	FFrameLogParser Parser;
+	Parser.ParseFrameLog(InTakeInfo.GetFrameLogFilePath(), InTakeInfo.VideoMetadata.FrameRate, InTakeInfo.DepthMetadata.FrameRate);
+
+	FTakeMetadata::FVideo Video;
+	Video.Name = "Video";
+	Video.FrameWidth = InTakeInfo.VideoMetadata.Resolution.X;
+	Video.FrameHeight = InTakeInfo.VideoMetadata.Resolution.Y;
+	Video.Orientation = ConvertOldOrientation(InTakeInfo.VideoMetadata.Orientation);
+	Video.FrameRate = InTakeInfo.VideoMetadata.FrameRate;
+	Video.FramesCount = static_cast<uint32>(InTakeInfo.TakeMetadata.NumFrames);
+	Video.DroppedFrames = TArray<uint32>{};
+	Video.Format = TEXT("mov");
+	Video.Path = InTakeInfo.GetVideoFilePath();
+	Video.PathType = FTakeMetadata::FVideo::EPathType::File;
+	Video.TimecodeStart = Parser.VideoTimecode.ToString();
+
+	NewTakeInfo.Video.Add(MoveTemp(Video));
+
+	if (HasDepthMetadataFile(InTakeInfo.TakeOriginDirectory))
+	{
+		FTakeMetadata::FVideo Depth;
+		Depth.Name = "Depth";
+		Depth.FramesCount = Parser.DepthFrameCount;
+		Depth.FrameRate = InTakeInfo.DepthMetadata.FrameRate;
+		Depth.FrameWidth = InTakeInfo.DepthMetadata.Resolution.X;
+		Depth.FrameHeight = InTakeInfo.DepthMetadata.Resolution.Y;
+		Depth.Orientation = ConvertOldOrientation(InTakeInfo.DepthMetadata.Orientation);
+		Depth.Format = TEXT("mha_depth");
+		Depth.Path = InTakeInfo.GetDepthFilePath();
+		Depth.PathType = FTakeMetadata::FVideo::EPathType::File;
+		Depth.TimecodeStart = Parser.DepthTimecode.ToString();
+
+		NewTakeInfo.Depth.Add(MoveTemp(Depth));
+		NewTakeInfo.Calibration.Push(FTakeMetadata::FCalibration{ TEXT("undefined"), InTakeInfo.GetCameraCalibrationFilePath(), TEXT("mhaical") });
+	}
+
+	FTakeMetadata::FAudio Audio;
+	Audio.Name = "Audio";
+	Audio.Path = InTakeInfo.GetVideoFilePath();
+
+	// We do not know the audio duration and so we must estimate it from the video. This apporach has been deemed 
+	// acceptable for the moment, based on how this duration value gets used.
+	Audio.Duration = static_cast<float>(InTakeInfo.TakeMetadata.NumFrames) / InTakeInfo.VideoMetadata.FrameRate;
+	Audio.TimecodeRate = Parser.AudioTimecodeRate.AsDecimal();
+	Audio.TimecodeStart = Parser.AudioTimecode.ToString();
+
+	TArray< FTakeMetadata::FAudio> Audios = { Audio };
+	NewTakeInfo.Audio = Audios;
+
+	return NewTakeInfo;
+}
+
+TOptional<FText> TakeDurationExceedsLimit(const float InDurationInSeconds)
+{
+	IConsoleVariable* CVar = IConsoleManager::Get().FindConsoleVariable(TEXT("au.SoundWaveImportLengthLimitInSeconds"));
+	if (!CVar)
+	{
+		return {};
+	}
+
+	static constexpr float Unlimited = -1.f;
+	float Limit = CVar->GetFloat();
+
+	if (FMath::IsNearlyEqual(Limit, Unlimited) && FMath::IsNegativeOrNegativeZero(Limit - InDurationInSeconds))
+	{
+		return {};
+	}
+
+	const FText Message = LOCTEXT("TakeDurationExceedsLimit", "Take duration ({0} seconds) exceeds allowed limit ({1} seconds).");
+
+	FNumberFormattingOptions Options;
+	Options.MaximumFractionalDigits = 2;
+	Options.MinimumFractionalDigits = 2;
+
+	return FText::Format(Message, FText::AsNumber(InDurationInSeconds, &Options), FText::AsNumber(Limit, &Options));
+}
+
+// If the video file derived from take.json metadata does not exist on disk, scan the directory
+// for a single .mov file and use it instead. Returns false only when multiple .mov files are
+// found, as the correct one cannot be determined.
+static bool TryResolveVideoFilePath(const FString& InTakeDirectory, FLiveLinkFaceTakeInfo& OutTakeInfo)
+{
+	const FString ExpectedVideoPath = OutTakeInfo.GetVideoFilePath();
+	if (IFileManager::Get().FileExists(*ExpectedVideoPath))
+	{
+		return true;
+	}
+
+	TArray<FString> MovFiles;
+	IFileManager::Get().FindFiles(MovFiles, *(InTakeDirectory / (TEXT("*") + FLiveLinkFaceStaticFileNames::VideoExt)), true, false);
+
+	if (MovFiles.Num() > 1)
+	{
+		UE_LOGF(LogLiveLinkFaceMetadata, Warning,
+			"Multiple .mov files found in take directory '%ls'; cannot determine correct video file for ingest",
+			*InTakeDirectory);
+		return false;
+	}
+
+	if (MovFiles.Num() == 1)
+	{
+		UE_LOGF(LogLiveLinkFaceMetadata, Log,
+			"Expected video file '%ls' not found; using '%ls' instead",
+			*ExpectedVideoPath, *(InTakeDirectory / MovFiles[0]));
+		OutTakeInfo.VideoFilePathOverride = InTakeDirectory / MovFiles[0];
+	}
+
+	// If no .mov found at all, leave the path as-is. Downstream ingest will surface the missing file error.
+	return true;
+}
+
+#define REPORT_IF_ERROR_AND_RETURN(Result, Text) if (!Result) { OutValidationErrors.Add(Text); return {};}
+#define REPORT_IF_ERROR(Result, Text) if (!Result) { OutValidationErrors.Add(Text); }
+
+namespace LiveLinkMetadata {
+
+TOptional<FTakeMetadata> ParseOldLiveLinkTakeMetadata(const FString& InTakeDirectory, TArray<FText>& OutValidationErrors)
+{
+	FLiveLinkFaceTakeInfo TakeInfo;
+
+	REPORT_IF_ERROR_AND_RETURN(FLiveLinkFaceMetadataParser::ParseTakeInfo(InTakeDirectory, TakeInfo), LOCTEXT("ParseOldTakeMetadata_ParseTakeInfoFailed", "Failed to parse take metadata"));
+	REPORT_IF_ERROR_AND_RETURN(FLiveLinkFaceMetadataParser::ParseVideoMetadata(InTakeDirectory, TakeInfo.VideoMetadata), LOCTEXT("ParseOldTakeMetadata_ParseVideoInfoFailed", "Failed to parse take video metadata"));
+	REPORT_IF_ERROR_AND_RETURN(FLiveLinkFaceMetadataParser::ParseAudioMetadata(InTakeDirectory, TakeInfo.AudioMetadata), LOCTEXT("ParseOldTakeMetadata_ParseAudioInfoFailed", "Failed to parse take audio metadata"));
+
+	// Shouldn't block the ingest process if this fails
+	FLiveLinkFaceMetadataParser::ParseThumbnail(InTakeDirectory, TakeInfo);
+
+	REPORT_IF_ERROR_AND_RETURN(TryResolveVideoFilePath(InTakeDirectory, TakeInfo), LOCTEXT("ParseOldTakeMetadata_MultipleMovFiles", "Multiple .mov files found in take directory; cannot determine correct video file"));
+
+	REPORT_IF_ERROR(!ContainsWhitespace(TakeInfo.GetTakeName()), LOCTEXT("ParseOldTakeMetadata_TakeNameContainsWhiteSpace", "Take name contains white space character(s)"));
+	REPORT_IF_ERROR(!ContainsWhitespace(TakeInfo.TakeMetadata.Subject), LOCTEXT("ParseOldTakeMetadata_SubjectContainsWhiteSpace", "Subject contains white space character(s)"));
+	REPORT_IF_ERROR(!ContainsWhitespace(TakeInfo.TakeMetadata.SlateName), LOCTEXT("ParseOldTakeMetadata_SlateNameContainsWhiteSpace", "Slate name contains white space character(s)"));
+
+	REPORT_IF_ERROR(FCString::IsPureAnsi(*InTakeDirectory), LOCTEXT("ParseOldTakeMetadata_UnsupportedCharactersInTakeDirectoryPath", "Take path contains unsupported text characters"));
+
+	const FString& SlateName = TakeInfo.TakeMetadata.SlateName;
+	REPORT_IF_ERROR(FCString::IsPureAnsi(*SlateName), FText::Format(LOCTEXT("ParseOldTakeMetadata_UnsupportedCharactersInSlateName", "Slate name '{0}' contains unsupported text characters"), FText::FromString(SlateName)));
+
+	const FString& Subject = TakeInfo.TakeMetadata.Subject;
+	REPORT_IF_ERROR(FCString::IsPureAnsi(*Subject), FText::Format(LOCTEXT("ParseOldTakeMetadata_UnsupportedCharactersInSubjectName", "Subject name '{0}' contains unsupported text characters"), FText::FromString(Subject)));
+
+	REPORT_IF_ERROR(FLiveLinkFaceMetadataParser::ParseDepthMetadata(InTakeDirectory, TakeInfo.DepthMetadata), LOCTEXT("ParseOldTakeMetadata_ParseDepthInfoFailed", "Failed to parse take depth metadata"));
+
+	TOptional<FText> TakeDurationResult = TakeDurationExceedsLimit(static_cast<float>(TakeInfo.TakeMetadata.NumFrames) / TakeInfo.VideoMetadata.FrameRate);
+	if (TakeDurationResult.IsSet())
+	{
+		OutValidationErrors.Add(TakeDurationResult.GetValue());
+	}
+
+	return ConvertOldToNewTakeMetadata(TakeInfo);
+}
+
+TArray<FTakeMetadata::FVideo> ParseOldLiveLinkVideoMetadataFromString(const FString& InJsonString, TArray<FText>& OutValidationErrors)
+{
+	FLiveLinkFaceTakeInfo TakeInfo;
+	REPORT_IF_ERROR_AND_RETURN(FLiveLinkFaceMetadataParser::ParseVideoMetadataFromString(InJsonString, TakeInfo.VideoMetadata), LOCTEXT("ParseOldTakeMetadata_ParseVideoInfoFailed", "Failed to parse take video metadata"));
+
+	return ConvertOldToNewTakeMetadata(TakeInfo).Video;
+}
+
+}
+
+}
+
+#undef LOCTEXT_NAMESPACE
diff --git a/Engine/Plugins/VirtualProduction/CaptureManager/CaptureManagerCore/Source/CaptureMetadataExtraction/Internal/LiveLinkFaceMetadata.h b/Engine/Plugins/VirtualProduction/CaptureManager/CaptureManagerCore/Source/CaptureMetadataExtraction/Internal/LiveLinkFaceMetadata.h
new file mode 100644
--- /dev/null
+++ b/Engine/Plugins/VirtualProduction/CaptureManager/CaptureManagerCore/Source/CaptureMetadataExtraction/Internal/LiveLinkFaceMetadata.h
@@ -0,0 +1,21 @@
+// Copyright Epic Games, Inc. All Rights Reserved.
+
+#pragma once
+
+#include "Dom/JsonObject.h"
+#include "IMediaTextureSample.h"
+
+#include "Containers/UnrealString.h"
+
+#include "CaptureManagerTakeMetadata.h"
+
+#define UE_API CAPTUREMETADATAEXTRACTION_API
+
+namespace UE::CaptureManager::LiveLinkMetadata
+{
+
+UE_API TOptional<FTakeMetadata> ParseOldLiveLinkTakeMetadata(const FString& InJsonFile, TArray<FText>& OutValidationError);
+UE_API TArray<FTakeMetadata::FVideo> ParseOldLiveLinkVideoMetadataFromString(const FString& InJsonString, TArray<FText>& OutValidationError);
+
+}
+#undef UE_API
diff --git a/Engine/Plugins/VirtualProduction/CaptureManager/CaptureManagerCore/Source/CaptureMetadataExtraction/Internal/LiveLinkFaceMetadataExtractor.cpp b/Engine/Plugins/VirtualProduction/CaptureManager/CaptureManagerCore/Source/CaptureMetadataExtraction/Internal/LiveLinkFaceMetadataExtractor.cpp
new file mode 100644
--- /dev/null
+++ b/Engine/Plugins/VirtualProduction/CaptureManager/CaptureManagerCore/Source/CaptureMetadataExtraction/Internal/LiveLinkFaceMetadataExtractor.cpp
@@ -0,0 +1,82 @@
+// Copyright Epic Games, Inc. All Rights Reserved.
+
+#include "LiveLinkFaceMetadataExtractor.h"
+
+#include "TakeArchiveMetadataExtractor.h"
+#include "LiveLinkFaceMetadata.h"
+#include "StereoCameraMetadataParseUtils.h"
+
+#include "HAL/FileManager.h"
+#include "Misc/Paths.h"
+
+namespace UE::CaptureManager
+{
+
+TValueOrError<FTakeMetadata, ELiveLinkFaceExtractionError>
+ExtractLiveLinkFaceMetadata(const FString& InTakeDirectoryPath, const FExtractionConfig& InConfig)
+{
+	if (InTakeDirectoryPath.IsEmpty())
+	{
+		return MakeError(ELiveLinkFaceExtractionError::DirectoryNotFound);
+	}
+
+	const FString TakeDirectoryPath = FPaths::ConvertRelativePathToFull(InTakeDirectoryPath);
+
+	if (!FPaths::DirectoryExists(TakeDirectoryPath))
+	{
+		return MakeError(ELiveLinkFaceExtractionError::DirectoryNotFound);
+	}
+
+	constexpr bool bFindFiles = true;
+	constexpr bool bFindDirectories = false;
+
+	TArray<FString> TakeFiles;
+	IFileManager::Get().FindFiles(
+		TakeFiles,
+		*(FPaths::Combine(TakeDirectoryPath, TEXT("*.") + FTakeMetadata::FileExtension)),
+		bFindFiles,
+		bFindDirectories
+	);
+
+	if (TakeFiles.Num() > 1)
+	{
+		return MakeError(ELiveLinkFaceExtractionError::MultipleMetadataFilesFound);
+	}
+
+	if (TakeFiles.Num() == 1)
+	{
+		TValueOrError<FTakeMetadata, ETakeArchiveExtractionError> Result = ExtractTakeArchiveMetadata(
+			FPaths::Combine(TakeDirectoryPath, TakeFiles[0]),
+			InConfig
+		);
+
+		if (Result.HasValue())
+		{
+			return MakeValue(Result.StealValue());
+		}
+
+		return MakeError(ELiveLinkFaceExtractionError::MetadataFormatNotRecognized);
+	}
+
+	// Fall back to legacy formats. Try LiveLink first (take.json), then HMC/StereoCamera.
+	// The stereo fallback is needed because this extractor serves as a base for CPS devices
+	// including HMC, which used the old StereoCamera take.json format.
+	TArray<FText> ValidationFailures;
+	TOptional<FTakeMetadata> OldFormatResult = LiveLinkMetadata::ParseOldLiveLinkTakeMetadata(TakeDirectoryPath, ValidationFailures);
+
+	if (OldFormatResult.IsSet())
+	{
+		return MakeValue(MoveTemp(OldFormatResult.GetValue()));
+	}
+
+	OldFormatResult = StereoCameraMetadata::ParseOldStereoCameraMetadata(TakeDirectoryPath, ValidationFailures);
+
+	if (OldFormatResult.IsSet())
+	{
+		return MakeValue(MoveTemp(OldFormatResult.GetValue()));
+	}
+
+	return MakeError(ELiveLinkFaceExtractionError::MetadataFileNotFound);
+}
+
+} // namespace UE::CaptureManager
diff --git a/Engine/Plugins/VirtualProduction/CaptureManager/CaptureManagerCore/Source/CaptureMetadataExtraction/Internal/LiveLinkFaceMetadataExtractor.h b/Engine/Plugins/VirtualProduction/CaptureManager/CaptureManagerCore/Source/CaptureMetadataExtraction/Internal/LiveLinkFaceMetadataExtractor.h
new file mode 100644
--- /dev/null
+++ b/Engine/Plugins/VirtualProduction/CaptureManager/CaptureManagerCore/Source/CaptureMetadataExtraction/Internal/LiveLinkFaceMetadataExtractor.h
@@ -0,0 +1,28 @@
+// Copyright Epic Games, Inc. All Rights Reserved.
+
+#pragma once
+
+#include "CaptureManagerExtractionConfig.h"
+#include "CaptureManagerTakeMetadata.h"
+
+#define UE_API CAPTUREMETADATAEXTRACTION_API
+
+namespace UE::CaptureManager
+{
+
+enum class ELiveLinkFaceExtractionError : uint8
+{
+	DirectoryNotFound,
+	MetadataFileNotFound,        // directory exists but contains no .cptake or legacy metadata
+	MetadataFormatNotRecognized, // metadata file found but could not be parsed
+	MultipleMetadataFilesFound,  // directory contains more than one .cptake file; legacy JSON is not multiplicity-checked
+};
+
+// Extracts take metadata from a Live Link Face capture directory.
+// Supports both .cptake and legacy directory layouts.
+UE_API TValueOrError<FTakeMetadata, ELiveLinkFaceExtractionError>
+ExtractLiveLinkFaceMetadata(const FString& InTakeDirectoryPath, const FExtractionConfig& InConfig);
+
+} // namespace UE::CaptureManager
+
+#undef UE_API
diff --git a/Engine/Plugins/VirtualProduction/CaptureManager/CaptureManagerCore/Source/CaptureMetadataExtraction/Internal/MonoVideoMetadataExtractor.cpp b/Engine/Plugins/VirtualProduction/CaptureManager/CaptureManagerCore/Source/CaptureMetadataExtraction/Internal/MonoVideoMetadataExtractor.cpp
new file mode 100644
--- /dev/null
+++ b/Engine/Plugins/VirtualProduction/CaptureManager/CaptureManagerCore/Source/CaptureMetadataExtraction/Internal/MonoVideoMetadataExtractor.cpp
@@ -0,0 +1,144 @@
+// Copyright Epic Games, Inc. All Rights Reserved.
+
+#include "MonoVideoMetadataExtractor.h"
+
+#include "Asset/CaptureAssetSanitization.h"
+#include "CaptureManagerFileExtensions.h"
+#include "Utils/CaptureExtractTimecode.h"
+#include "Utils/ParseTakeUtils.h"
+
+#include "HAL/FileManager.h"
+#include "Misc/Guid.h"
+#include "Misc/Paths.h"
+
+DEFINE_LOG_CATEGORY_STATIC(LogMonoVideoExtractor, Log, All);
+
+namespace UE::CaptureManager
+{
+
+TValueOrError<FTakeMetadata, EMonoVideoExtractionError>
+ExtractMonoVideoMetadata(FMonoVideoDescriptor InDescriptor, const FExtractionConfig& InConfig)
+{
+	// Resolve absolute paths
+
+	InDescriptor.VideoFilePath = FPaths::ConvertRelativePathToFull(InDescriptor.VideoFilePath);
+	for (FString& AudioPath : InDescriptor.AudioFilePaths)
+	{
+		AudioPath = FPaths::ConvertRelativePathToFull(AudioPath);
+	}
+
+	if (!FPaths::FileExists(InDescriptor.VideoFilePath))
+	{
+		return MakeError(EMonoVideoExtractionError::VideoFileNotFound);
+	}
+
+	const FString Extension = FPaths::GetExtension(InDescriptor.VideoFilePath);
+	if (!IsVideoExtension(Extension))
+	{
+		return MakeError(EMonoVideoExtractionError::UnsupportedVideoFormat);
+	}
+
+	for (const FString& AudioPath : InDescriptor.AudioFilePaths)
+	{
+		if (!FPaths::FileExists(AudioPath))
+		{
+			return MakeError(EMonoVideoExtractionError::AudioFileNotFound);
+		}
+	}
+
+	// Build metadata from known inputs before attempting FFprobe enrichment
+	FTakeMetadata Metadata;
+	Metadata.Version.Major = 4;
+	Metadata.Version.Minor = 1;
+	FString Slate = InDescriptor.Slate.IsEmpty() ? FPaths::GetBaseFilename(InDescriptor.VideoFilePath) : InDescriptor.Slate;
+	SanitizePackagePath(Slate, TEXT('_'));
+	Metadata.Slate = MoveTemp(Slate);
+	Metadata.TakeNumber = (InDescriptor.TakeNumber >= 1) ? static_cast<uint32>(InDescriptor.TakeNumber) : 1u;
+	Metadata.UniqueId = FGuid::NewGuid().ToString(EGuidFormats::DigitsWithHyphensLower);
+	Metadata.DateTime = IFileManager::Get().GetStatData(*InDescriptor.VideoFilePath).CreationTime;
+	Metadata.Device.Model = TEXT("MonoVideo");
+
+	FTakeMetadata::FVideo& Video = Metadata.Video.AddDefaulted_GetRef();
+	Video.Name = FPaths::GetBaseFilename(InDescriptor.VideoFilePath);
+	Video.Format = Extension.ToLower();
+	Video.Path = InDescriptor.VideoFilePath;
+	Video.PathType = FTakeMetadata::FVideo::EPathType::File;
+
+	// Attempt to enrich the metadata using ffprobe/electra
+
+	TOptional<Private::FFProbeCommand> ProbeCommand;
+
+	if (InConfig.bUseFFprobe && !InConfig.FFmpegPath.IsEmpty())
+	{
+		ProbeCommand = Private::FFProbeCommand(InConfig.FFmpegPath);
+	}
+
+	FCaptureExtractVideoInfo::FResult VideoInfoResult = FCaptureExtractVideoInfo::Create(InDescriptor.VideoFilePath, ProbeCommand);
+
+	if (VideoInfoResult.IsValid())
+	{
+		FCaptureExtractVideoInfo VideoInfo = VideoInfoResult.StealValue();
+
+		const FString TimecodeStartString = VideoInfo.GetTimecode().ToString();
+		Video.TimecodeStart = TimecodeStartString;
+		Video.FrameRate = static_cast<float>(VideoInfo.GetFrameRate().AsDecimal());
+
+		switch (VideoInfo.GetVideoOrientation())
+		{
+		case EMediaOrientation::CW90:
+			Video.Orientation = FTakeMetadata::FVideo::EOrientation::CW90;
+			break;
+		case EMediaOrientation::CW180:
+			Video.Orientation = FTakeMetadata::FVideo::EOrientation::CW180;
+			break;
+		case EMediaOrientation::CW270:
+			Video.Orientation = FTakeMetadata::FVideo::EOrientation::CW270;
+			break;
+		default:
+			Video.Orientation = FTakeMetadata::FVideo::EOrientation::Original;
+		}
+
+		if (InDescriptor.AudioFilePaths.IsEmpty() && VideoInfo.ContainsAudio())
+		{
+			FTakeMetadata::FAudio& Audio = Metadata.Audio.AddDefaulted_GetRef();
+			Audio.Name = FPaths::GetBaseFilename(InDescriptor.VideoFilePath);
+			Audio.Path = InDescriptor.VideoFilePath;
+			Audio.Duration = VideoInfo.GetAudioDurationSeconds();
+			Audio.TimecodeStart = TimecodeStartString;
+			Audio.TimecodeRate = Video.FrameRate;
+		}
+	}
+	else
+	{
+		UE_LOGF(
+			LogMonoVideoExtractor,
+			Warning,
+			"Failed to extract video info from '%ls'. Timecode, frame rate, orientation, and embedded audio detection unavailable.",
+			*InDescriptor.VideoFilePath
+		);
+	}
+
+	// External audio - probe independently for timecodes
+	const FFrameRate VideoFrameRate = Video.FrameRate > 0.0f ? UE::CaptureManager::ParseFrameRate(Video.FrameRate) : FFrameRate();
+
+	for (const FString& AudioPath : InDescriptor.AudioFilePaths)
+	{
+		FTakeMetadata::FAudio& Audio = Metadata.Audio.AddDefaulted_GetRef();
+		Audio.Name = FPaths::GetBaseFilename(AudioPath);
+		Audio.Path = AudioPath;
+
+		TSharedPtr<FCaptureExtractAudioTimecode> AudioExtractor = MakeShareable(new FCaptureExtractAudioTimecode(AudioPath));
+		FCaptureExtractAudioTimecode::FTimecodeInfoResult TimecodeResult = AudioExtractor->Extract(VideoFrameRate);
+
+		if (TimecodeResult.IsValid())
+		{
+			FTimecodeInfo TimecodeInfo = TimecodeResult.GetValue();
+			Audio.TimecodeStart = TimecodeInfo.Timecode.ToString();
+			Audio.TimecodeRate = static_cast<float>(TimecodeInfo.TimecodeRate.AsDecimal());
+		}
+	}
+
+	return MakeValue(MoveTemp(Metadata));
+}
+
+} // namespace UE::CaptureManager
diff --git a/Engine/Plugins/VirtualProduction/CaptureManager/CaptureManagerCore/Source/CaptureMetadataExtraction/Internal/MonoVideoMetadataExtractor.h b/Engine/Plugins/VirtualProduction/CaptureManager/CaptureManagerCore/Source/CaptureMetadataExtraction/Internal/MonoVideoMetadataExtractor.h
new file mode 100644
--- /dev/null
+++ b/Engine/Plugins/VirtualProduction/CaptureManager/CaptureManagerCore/Source/CaptureMetadataExtraction/Internal/MonoVideoMetadataExtractor.h
@@ -0,0 +1,35 @@
+// Copyright Epic Games, Inc. All Rights Reserved.
+
+#pragma once
+
+#include "CaptureManagerExtractionConfig.h"
+#include "CaptureManagerTakeMetadata.h"
+
+#define UE_API CAPTUREMETADATAEXTRACTION_API
+
+namespace UE::CaptureManager
+{
+
+enum class EMonoVideoExtractionError : uint8
+{
+	VideoFileNotFound,
+	UnsupportedVideoFormat,
+	AudioFileNotFound,
+};
+
+struct FMonoVideoDescriptor
+{
+	FString VideoFilePath;
+	TArray<FString> AudioFilePaths; // empty = use embedded audio if present
+	FString Slate;                  // empty = derived from video filename stem
+	int32 TakeNumber = 1;           // clamped to 1 if < 1
+};
+
+// Media probing (timecodes, frame rate, orientation) is best-effort.
+// Probe failure is non-fatal - extraction succeeds with missing values.
+UE_API TValueOrError<FTakeMetadata, EMonoVideoExtractionError>
+ExtractMonoVideoMetadata(FMonoVideoDescriptor InDescriptor, const FExtractionConfig& InConfig);
+
+} // namespace UE::CaptureManager
+
+#undef UE_API
diff --git a/Engine/Plugins/VirtualProduction/CaptureManager/CaptureManagerCore/Source/CaptureMetadataExtraction/Internal/StereoVideoMetadataExtractor.cpp b/Engine/Plugins/VirtualProduction/CaptureManager/CaptureManagerCore/Source/CaptureMetadataExtraction/Internal/StereoVideoMetadataExtractor.cpp
new file mode 100644
--- /dev/null
+++ b/Engine/Plugins/VirtualProduction/CaptureManager/CaptureManagerCore/Source/CaptureMetadataExtraction/Internal/StereoVideoMetadataExtractor.cpp
@@ -0,0 +1,317 @@
+// Copyright Epic Games, Inc. All Rights Reserved.
+
+#include "StereoVideoMetadataExtractor.h"
+#include "CaptureManagerCalibrationUtils.h"
+#include "CaptureManagerFileExtensions.h"
+
+#include "Utils/CaptureExtractTimecode.h"
+#include "Asset/CaptureAssetSanitization.h"
+
+#include "HAL/FileManager.h"
+#include "ImageUtils.h"
+#include "Misc/FileHelper.h"
+#include "Serialization/JsonSerializer.h"
+#include "Misc/Guid.h"
+#include "Misc/Paths.h"
+
+DEFINE_LOG_CATEGORY_STATIC(LogStereoVideoExtractor, Log, All);
+
+namespace UE::CaptureManager
+{
+
+namespace Private
+{
+
+enum class EVideoInputType
+{
+	File,
+	ImageSequence,
+};
+
+static TValueOrError<EVideoInputType, EStereoVideoExtractionError> ResolveVideoInputType(
+	const FString& InPath,
+	EStereoVideoExtractionError NotFoundError,
+	EStereoVideoExtractionError UnsupportedFormatError
+)
+{
+	if (InPath.IsEmpty())
+	{
+		return MakeError(NotFoundError);
+	}
+
+	if (FPaths::DirectoryExists(InPath))
+	{
+		return MakeValue(EVideoInputType::ImageSequence);
+	}
+
+	if (FPaths::FileExists(InPath))
+	{
+		if (IsVideoExtension(FPaths::GetExtension(InPath)))
+		{
+			return MakeValue(EVideoInputType::File);
+		}
+		return MakeError(UnsupportedFormatError);
+	}
+
+	return MakeError(NotFoundError);
+}
+
+} // namespace Private
+
+TValueOrError<FTakeMetadata, EStereoVideoExtractionError>
+ExtractStereoVideoMetadata(FStereoVideoDescriptor InDescriptor, const FExtractionConfig& InConfig)
+{
+	// Resolve absolute paths
+
+	if (!InDescriptor.VideoPathA.IsEmpty())
+	{
+		InDescriptor.VideoPathA = FPaths::ConvertRelativePathToFull(InDescriptor.VideoPathA);
+	}
+	if (!InDescriptor.VideoPathB.IsEmpty())
+	{
+		InDescriptor.VideoPathB = FPaths::ConvertRelativePathToFull(InDescriptor.VideoPathB);
+	}
+	for (FString& AudioPath : InDescriptor.AudioFilePaths)
+	{
+		AudioPath = FPaths::ConvertRelativePathToFull(AudioPath);
+	}
+	if (!InDescriptor.CalibrationFilePath.IsEmpty())
+	{
+		InDescriptor.CalibrationFilePath = FPaths::ConvertRelativePathToFull(InDescriptor.CalibrationFilePath);
+	}
+
+	// Validate and classify inputs
+
+	TValueOrError<Private::EVideoInputType, EStereoVideoExtractionError> TypeA = Private::ResolveVideoInputType(
+		InDescriptor.VideoPathA,
+		EStereoVideoExtractionError::VideoANotFound,
+		EStereoVideoExtractionError::UnsupportedVideoFormatA
+	);
+
+	if (TypeA.HasError())
+	{
+		return MakeError(TypeA.GetError());
+	}
+
+	TValueOrError<Private::EVideoInputType, EStereoVideoExtractionError> TypeB = Private::ResolveVideoInputType(
+		InDescriptor.VideoPathB,
+		EStereoVideoExtractionError::VideoBNotFound,
+		EStereoVideoExtractionError::UnsupportedVideoFormatB
+	);
+
+	if (TypeB.HasError())
+	{
+		return MakeError(TypeB.GetError());
+	}
+
+	if (TypeA.GetValue() != TypeB.GetValue())
+	{
+		return MakeError(EStereoVideoExtractionError::VideoTypeMismatch);
+	}
+
+	for (const FString& AudioPath : InDescriptor.AudioFilePaths)
+	{
+		if (!FPaths::FileExists(AudioPath))
+		{
+			return MakeError(EStereoVideoExtractionError::AudioFileNotFound);
+		}
+	}
+
+	if (!InDescriptor.CalibrationFilePath.IsEmpty() && !FPaths::FileExists(InDescriptor.CalibrationFilePath))
+	{
+		return MakeError(EStereoVideoExtractionError::CalibrationFileNotFound);
+	}
+
+	// Resolve calibration format
+
+	FString ResolvedCalibrationFormat;
+	if (!InDescriptor.CalibrationFilePath.IsEmpty())
+	{
+		ResolvedCalibrationFormat = DetectCalibrationFormat(InDescriptor.CalibrationFilePath);
+
+		UE_LOGF(
+			LogStereoVideoExtractor,
+			Verbose,
+			"Detected calibration format '%ls' for: %ls",
+			*ResolvedCalibrationFormat, *InDescriptor.CalibrationFilePath
+		);
+
+		if (ResolvedCalibrationFormat.IsEmpty())
+		{
+			return MakeError(EStereoVideoExtractionError::CalibrationFormatUnrecognized);
+		}
+	}
+
+	// Build metadata
+
+	TOptional<Private::FFProbeCommand> ProbeCommand;
+	if (InConfig.bUseFFprobe && !InConfig.FFmpegPath.IsEmpty())
+	{
+		ProbeCommand = Private::FFProbeCommand(InConfig.FFmpegPath);
+	}
+
+	FTakeMetadata Metadata;
+	Metadata.Version.Major = 4;
+	Metadata.Version.Minor = 1;
+	Metadata.TakeNumber = (InDescriptor.TakeNumber >= 1) ? static_cast<uint32>(InDescriptor.TakeNumber) : 1u;
+	Metadata.UniqueId = FGuid::NewGuid().ToString(EGuidFormats::DigitsWithHyphensLower);
+	Metadata.Device.Model = TEXT("StereoHMC");
+
+	const Private::EVideoInputType VideoType = TypeA.GetValue();
+
+	bool bVideoFrameRateSet = false;
+	FFrameRate VideoFrameRate;
+
+	const FString VideoPaths[] = { InDescriptor.VideoPathA, InDescriptor.VideoPathB };
+
+	for (int32 VideoIndex = 0; VideoIndex < 2; ++VideoIndex)
+	{
+		const FString& VideoPath = VideoPaths[VideoIndex];
+		FTakeMetadata::FVideo Video;
+
+		// Derive name from filename stem (video file) or folder name (image sequence)
+		FString Name;
+		// FPaths::GetBaseFilename strips extension for files, returns leaf component for directories
+		Name = FPaths::GetBaseFilename(VideoPath);
+		SanitizePackagePath(Name, TEXT('_'));
+		Video.Name = MoveTemp(Name);
+
+		if (VideoType == Private::EVideoInputType::File)
+		{
+			Video.Format = FPaths::GetExtension(VideoPath).ToLower();
+			Video.Path = VideoPath;
+			Video.PathType = FTakeMetadata::FVideo::EPathType::File;
+
+			FCaptureExtractVideoInfo::FResult ExtractResult = FCaptureExtractVideoInfo::Create(VideoPath, ProbeCommand);
+			if (ExtractResult.IsValid())
+			{
+				FCaptureExtractVideoInfo VideoInfo = ExtractResult.StealValue();
+				Video.TimecodeStart = VideoInfo.GetTimecode().ToString();
+				Video.FrameRate = static_cast<float>(VideoInfo.GetFrameRate().AsDecimal());
+
+				if (!bVideoFrameRateSet)
+				{
+					VideoFrameRate = VideoInfo.GetFrameRate();
+					bVideoFrameRateSet = true;
+				}
+
+				switch (VideoInfo.GetVideoOrientation())
+				{
+				case EMediaOrientation::CW90:
+					Video.Orientation = FTakeMetadata::FVideo::EOrientation::CW90;
+					break;
+				case EMediaOrientation::CW180:
+					Video.Orientation = FTakeMetadata::FVideo::EOrientation::CW180;
+					break;
+				case EMediaOrientation::CW270:
+					Video.Orientation = FTakeMetadata::FVideo::EOrientation::CW270;
+					break;
+				default:
+					Video.Orientation = FTakeMetadata::FVideo::EOrientation::Original;
+				}
+			}
+			else
+			{
+				UE_LOGF(
+					LogStereoVideoExtractor,
+					Warning,
+					"Failed to extract video info from '%ls'. Timecode, frame rate, and orientation unavailable for this video.",
+					*VideoPath
+				);
+			}
+
+			// Use creation time of video A for the take timestamp
+			if (VideoIndex == 0)
+			{
+				Metadata.DateTime = IFileManager::Get().GetStatData(*VideoPath).CreationTime;
+			}
+		}
+		else // IMAGE_SEQUENCE
+		{
+			Video.Path = VideoPath;
+			Video.PathType = FTakeMetadata::FVideo::EPathType::Folder;
+
+			int32 FramesCount = 0;
+			bool bFirstFrame = true;
+
+			IFileManager::Get().IterateDirectoryRecursively(*VideoPath,
+				[&](const TCHAR* InPath, bool bIsDirectory) -> bool
+				{
+					if (!bIsDirectory && IsImageExtension(FPaths::GetExtension(InPath)))
+					{
+						if (bFirstFrame)
+						{
+							bFirstFrame = false;
+
+							FString FilePath, Filename, Ext;
+							FPaths::Split(InPath, FilePath, Filename, Ext);
+							Video.Format = Ext.ToLower();
+							Video.Path = VideoPath;
+
+							FImage Image;
+							FImageUtils::LoadImage(InPath, Image);
+							if (Image.IsImageInfoValid())
+							{
+								Video.FrameWidth = static_cast<uint32>(Image.SizeX);
+								Video.FrameHeight = static_cast<uint32>(Image.SizeY);
+							}
+
+							if (VideoIndex == 0)
+							{
+								Metadata.DateTime = IFileManager::Get().GetStatData(InPath).CreationTime;
+							}
+						}
+						FramesCount++;
+					}
+					return true;
+				});
+
+			Video.FramesCount = FramesCount;
+		}
+
+		Metadata.Video.Add(MoveTemp(Video));
+	}
+
+	Metadata.Slate = InDescriptor.Slate.IsEmpty() ? Metadata.Video[0].Name : InDescriptor.Slate;
+
+	// Audio
+
+	if (InDescriptor.AudioFilePaths.Num() > 0 && !bVideoFrameRateSet)
+	{
+		UE_LOGF(LogStereoVideoExtractor, Warning, "Video frame rate unavailable; audio timecode may be set incorrectly.");
+	}
+
+	for (const FString& AudioPath : InDescriptor.AudioFilePaths)
+	{
+		FTakeMetadata::FAudio Audio;
+		Audio.Name = FPaths::GetBaseFilename(AudioPath);
+		Audio.Path = AudioPath;
+		Audio.Duration = 0;
+
+		TSharedPtr<FCaptureExtractAudioTimecode> Extractor = MakeShareable(new FCaptureExtractAudioTimecode(AudioPath));
+		FCaptureExtractAudioTimecode::FTimecodeInfoResult TimecodeResult = Extractor->Extract(VideoFrameRate);
+		if (TimecodeResult.IsValid())
+		{
+			FTimecodeInfo TimecodeInfo = TimecodeResult.GetValue();
+			Audio.TimecodeStart = TimecodeInfo.Timecode.ToString();
+			Audio.TimecodeRate = static_cast<float>(TimecodeInfo.TimecodeRate.AsDecimal());
+		}
+
+		Metadata.Audio.Add(MoveTemp(Audio));
+	}
+
+	// Calibration
+
+	if (!InDescriptor.CalibrationFilePath.IsEmpty())
+	{
+		FTakeMetadata::FCalibration Calibration;
+		Calibration.Name = TEXT("calibration");
+		Calibration.Path = InDescriptor.CalibrationFilePath;
+		Calibration.Format = ResolvedCalibrationFormat;
+		Metadata.Calibration.Add(MoveTemp(Calibration));
+	}
+
+	return MakeValue(MoveTemp(Metadata));
+}
+
+} // namespace UE::CaptureManager
diff --git a/Engine/Plugins/VirtualProduction/CaptureManager/CaptureManagerCore/Source/CaptureMetadataExtraction/Internal/StereoVideoMetadataExtractor.h b/Engine/Plugins/VirtualProduction/CaptureManager/CaptureManagerCore/Source/CaptureMetadataExtraction/Internal/StereoVideoMetadataExtractor.h
new file mode 100644
--- /dev/null
+++ b/Engine/Plugins/VirtualProduction/CaptureManager/CaptureManagerCore/Source/CaptureMetadataExtraction/Internal/StereoVideoMetadataExtractor.h
@@ -0,0 +1,42 @@
+// Copyright Epic Games, Inc. All Rights Reserved.
+
+#pragma once
+
+#include "CaptureManagerExtractionConfig.h"
+#include "CaptureManagerTakeMetadata.h"
+
+#define UE_API CAPTUREMETADATAEXTRACTION_API
+
+namespace UE::CaptureManager
+{
+
+enum class EStereoVideoExtractionError : uint8
+{
+	VideoANotFound,
+	VideoBNotFound,
+	UnsupportedVideoFormatA,     // video A: not a supported video file or image sequence directory
+	UnsupportedVideoFormatB,     // video B: not a supported video file or image sequence directory
+	VideoTypeMismatch,           // video A is a file, video B is a folder, or vice versa
+	AudioFileNotFound,
+	CalibrationFileNotFound,
+	CalibrationFormatUnrecognized, // calibration file provided but format could not be detected
+};
+
+struct FStereoVideoDescriptor
+{
+	FString VideoPathA;              // video file (.mp4/.mov) or image sequence folder
+	FString VideoPathB;              // same type as video A
+	TArray<FString> AudioFilePaths;  // empty = no audio
+	FString CalibrationFilePath;     // empty = no calibration; format auto-detected from JSON structure
+	FString Slate;                   // empty = derived from video A stem/folder name
+	int32 TakeNumber = 1;            // clamped to 1 if < 1
+};
+
+// Media probing (timecodes, frame rate, orientation) is best-effort.
+// Probe failure is non-fatal - extraction succeeds with missing values.
+UE_API TValueOrError<FTakeMetadata, EStereoVideoExtractionError>
+ExtractStereoVideoMetadata(FStereoVideoDescriptor InDescriptor, const FExtractionConfig& InConfig);
+
+} // namespace UE::CaptureManager
+
+#undef UE_API
diff --git a/Engine/Plugins/VirtualProduction/CaptureManager/CaptureManagerCore/Source/CaptureMetadataExtraction/Internal/TakeArchiveMetadataExtractor.cpp b/Engine/Plugins/VirtualProduction/CaptureManager/CaptureManagerCore/Source/CaptureMetadataExtraction/Internal/TakeArchiveMetadataExtractor.cpp
new file mode 100644
--- /dev/null
+++ b/Engine/Plugins/VirtualProduction/CaptureManager/CaptureManagerCore/Source/CaptureMetadataExtraction/Internal/TakeArchiveMetadataExtractor.cpp
@@ -0,0 +1,173 @@
+// Copyright Epic Games, Inc. All Rights Reserved.
+
+#include "TakeArchiveMetadataExtractor.h"
+
+#include "LiveLinkFaceMetadata.h"
+#include "StereoCameraMetadataParseUtils.h"
+
+#include "Utils/CaptureExtractTimecode.h"
+#include "Utils/ParseTakeUtils.h"
+#include "Misc/Paths.h"
+
+DEFINE_LOG_CATEGORY_STATIC(LogTakeArchiveExtractor, Log, All);
+
+namespace UE::CaptureManager
+{
+static FString ErrorOriginToString(FTakeMetadataParserError::EOrigin InOrigin)
+{
+	switch (InOrigin)
+	{
+	case FTakeMetadataParserError::Reader:
+		return TEXT("Reader");
+	case FTakeMetadataParserError::Validator:
+		return TEXT("Validator");
+	case FTakeMetadataParserError::Parser:
+	default:
+		return TEXT("Parser");
+	}
+}
+
+TValueOrError<FTakeMetadata, ETakeArchiveExtractionError> ExtractTakeArchiveMetadata(const FString& InMetadataFilePath, const FExtractionConfig& InConfig)
+{
+	using namespace Private;
+
+	const FString MetadataFilePath = FPaths::ConvertRelativePathToFull(InMetadataFilePath);
+
+	if (!FPaths::FileExists(MetadataFilePath))
+	{
+		UE_LOGF(LogTakeArchiveExtractor, Error, "Metadata file does not exist: %ls", *MetadataFilePath);
+		return MakeError(ETakeArchiveExtractionError::MetadataFileNotFound);
+	}
+
+	FTakeMetadataParser TakeMetadataParser;
+	TValueOrError<FTakeMetadata, FTakeMetadataParserError> TakeMetadataResult = TakeMetadataParser.Parse(MetadataFilePath);
+	const FString CurrentDirectory = FPaths::GetPath(MetadataFilePath);
+
+	if (TakeMetadataResult.HasError())
+	{
+		const FTakeMetadataParserError ParserError = TakeMetadataResult.StealError();
+		const FString OriginName = UE::CaptureManager::ErrorOriginToString(ParserError.Origin);
+
+		const FString Message = FString::Format(
+			TEXT("Unable to parse take metadata file - {0} (Error origin: {1}): {2}"),
+			{
+				MetadataFilePath,
+				OriginName,
+				ParserError.Message.ToString()
+			}
+		);
+
+		UE_LOGF(LogTakeArchiveExtractor, Warning, "%ls", *Message);
+		UE_LOGF(LogTakeArchiveExtractor, Display, "Checking backwards compatible take metadata formats");
+
+		TArray<FText> ValidationFailures;
+
+		UE_LOGF(LogTakeArchiveExtractor, Display, "Checking directory (%ls) for legacy LiveLink take metadata format", *CurrentDirectory);
+		TOptional<FTakeMetadata> ParseOldMetadataResult = LiveLinkMetadata::ParseOldLiveLinkTakeMetadata(CurrentDirectory, ValidationFailures);
+
+		if (ParseOldMetadataResult.IsSet())
+		{
+			return MakeValue(MoveTemp(ParseOldMetadataResult.GetValue()));
+		}
+
+		UE_LOGF(LogTakeArchiveExtractor, Display, "Checking directory (%ls) for legacy StereoCamera take metadata format", *CurrentDirectory);
+		ParseOldMetadataResult = StereoCameraMetadata::ParseOldStereoCameraMetadata(CurrentDirectory, ValidationFailures);
+
+		if (ParseOldMetadataResult.IsSet())
+		{
+			return MakeValue(MoveTemp(ParseOldMetadataResult.GetValue()));
+		}
+
+		UE_LOGF(LogTakeArchiveExtractor, Error, "Failed to parse take metadata file: %ls", *CurrentDirectory);
+		return MakeError(ETakeArchiveExtractionError::MetadataFormatNotRecognized);
+	}
+
+	FTakeMetadata& TakeMetadata = TakeMetadataResult.GetValue();
+
+	const TOptional<FString> ThumbnailPath = TakeMetadata.Thumbnail.GetThumbnailPath();
+	if (ThumbnailPath.IsSet())
+	{
+		if (FPaths::IsRelative(*ThumbnailPath))
+		{
+			TakeMetadata.Thumbnail = FTakeThumbnailData(FPaths::ConvertRelativePathToFull(CurrentDirectory, *ThumbnailPath));
+		}
+	}
+
+	TOptional<FFProbeCommand> ProbeCommandBuilder;
+
+	if (InConfig.bUseFFprobe && !InConfig.FFmpegPath.IsEmpty())
+	{
+		ProbeCommandBuilder = FFProbeCommand(InConfig.FFmpegPath);
+	}
+
+	bool bVideoFrameRateSet = false;
+	FFrameRate VideoFrameRate;
+
+	for (FTakeMetadata::FVideo& Video : TakeMetadata.Video)
+	{
+		if (FPaths::IsRelative(*Video.Path))
+		{
+			Video.Path = FPaths::ConvertRelativePathToFull(CurrentDirectory, *Video.Path);
+		}
+
+		// Check if this is a file (and not an image sequence directory)
+		if (FPaths::FileExists(Video.Path) && !Video.TimecodeStart.IsSet())
+		{
+			FCaptureExtractVideoInfo::FResult ExtractorOpt = FCaptureExtractVideoInfo::Create(Video.Path, ProbeCommandBuilder);
+
+			if (ExtractorOpt.IsValid())
+			{
+				FCaptureExtractVideoInfo Extractor = ExtractorOpt.StealValue();
+				Video.TimecodeStart = Extractor.GetTimecode().ToString();
+				if (FMath::IsNearlyZero(Video.FrameRate))
+				{
+					Video.FrameRate = static_cast<float>(Extractor.GetFrameRate().AsDecimal());
+				}
+			}
+		}
+
+		if (!bVideoFrameRateSet)
+		{
+			VideoFrameRate = UE::CaptureManager::ParseFrameRate(Video.FrameRate);
+			bVideoFrameRateSet = true;
+		}
+	}
+
+	for (FTakeMetadata::FAudio& Audio : TakeMetadata.Audio)
+	{
+		if (FPaths::IsRelative(*Audio.Path))
+		{
+			Audio.Path = FPaths::ConvertRelativePathToFull(CurrentDirectory, *Audio.Path);
+		}
+
+		// Only probe when the .cptake file did not provide timecodes. The metadata file is authoritative.
+		if (!Audio.TimecodeStart.IsSet() && !Audio.TimecodeRate.IsSet())
+		{
+			TSharedPtr<FCaptureExtractAudioTimecode> Extractor = MakeShareable(new FCaptureExtractAudioTimecode(Audio.Path));
+
+			// The video frame rate will be used to calculate the timecode rate if the timecode rate cannot be extracted from the audio file
+			// If the VideoFrameRate is not set, the timecode rate may be set to the audio sample rate
+			FCaptureExtractAudioTimecode::FTimecodeInfoResult TimecodeInfoResult = Extractor->Extract(VideoFrameRate);
+
+			if (TimecodeInfoResult.IsValid())
+			{
+				const FTimecodeInfo TimecodeInfo = TimecodeInfoResult.GetValue();
+
+				Audio.TimecodeStart = TimecodeInfo.Timecode.ToString();
+				Audio.TimecodeRate = static_cast<float>(TimecodeInfo.TimecodeRate.AsDecimal());
+			}
+		}
+	}
+
+	for (FTakeMetadata::FCalibration& Calibration : TakeMetadata.Calibration)
+	{
+		if (FPaths::IsRelative(*Calibration.Path))
+		{
+			Calibration.Path = FPaths::ConvertRelativePathToFull(CurrentDirectory, *Calibration.Path);
+		}
+	}
+
+	return MakeValue(TakeMetadataResult.StealValue());
+}
+
+}
diff --git a/Engine/Plugins/VirtualProduction/CaptureManager/CaptureManagerCore/Source/CaptureMetadataExtraction/Internal/TakeArchiveMetadataExtractor.h b/Engine/Plugins/VirtualProduction/CaptureManager/CaptureManagerCore/Source/CaptureMetadataExtraction/Internal/TakeArchiveMetadataExtractor.h
new file mode 100644
--- /dev/null
+++ b/Engine/Plugins/VirtualProduction/CaptureManager/CaptureManagerCore/Source/CaptureMetadataExtraction/Internal/TakeArchiveMetadataExtractor.h
@@ -0,0 +1,25 @@
+// Copyright Epic Games, Inc. All Rights Reserved.
+
+#pragma once
+
+#include "CaptureManagerExtractionConfig.h"
+#include "CaptureManagerTakeMetadata.h"
+
+#define UE_API CAPTUREMETADATAEXTRACTION_API
+
+namespace UE::CaptureManager
+{
+
+enum class ETakeArchiveExtractionError : uint8
+{
+	MetadataFileNotFound,
+	MetadataFormatNotRecognized,
+};
+
+// Media probing (timecodes, frame rate, orientation) is best-effort.
+// Probe failure is non-fatal - extraction succeeds with missing values.
+UE_API TValueOrError<FTakeMetadata, ETakeArchiveExtractionError> ExtractTakeArchiveMetadata(const FString& InMetadataFilePath, const FExtractionConfig& InConfig);
+
+} // namespace UE::CaptureManager
+
+#undef UE_API
diff --git a/Engine/Plugins/VirtualProduction/CaptureManager/CaptureManagerCore/Source/CaptureMetadataExtraction/Private/CaptureMetadataExtractionModule.cpp b/Engine/Plugins/VirtualProduction/CaptureManager/CaptureManagerCore/Source/CaptureMetadataExtraction/Private/CaptureMetadataExtractionModule.cpp
new file mode 100644
--- /dev/null
+++ b/Engine/Plugins/VirtualProduction/CaptureManager/CaptureManagerCore/Source/CaptureMetadataExtraction/Private/CaptureMetadataExtractionModule.cpp
@@ -0,0 +1,5 @@
+// Copyright Epic Games, Inc. All Rights Reserved.
+
+#include "Modules/ModuleManager.h"
+
+IMPLEMENT_MODULE(FDefaultModuleImpl, CaptureMetadataExtraction)
diff --git a/Engine/Plugins/VirtualProduction/CaptureManager/CaptureManagerCore/Source/CaptureMetadataExtraction/Private/ResolutionResolver.cpp b/Engine/Plugins/VirtualProduction/CaptureManager/CaptureManagerCore/Source/CaptureMetadataExtraction/Private/ResolutionResolver.cpp
new file mode 100644
--- /dev/null
+++ b/Engine/Plugins/VirtualProduction/CaptureManager/CaptureManagerCore/Source/CaptureMetadataExtraction/Private/ResolutionResolver.cpp
@@ -0,0 +1,119 @@
+// Copyright Epic Games, Inc. All Rights Reserved.
+
+#include "ResolutionResolver.h"
+
+#include "IImageWrapper.h"
+#include "IImageWrapperModule.h"
+
+#include "Misc/FileHelper.h"
+#include "Misc/Paths.h"
+#include "Modules/ModuleManager.h"
+
+namespace UE::CaptureManager
+{
+
+TValueOrError<FIntPoint, FResolutionResolver::EAddError> FResolutionResolver::Add(const FStereoCameraTakeInfo::FCamera& InCamera)
+{
+	TValueOrError<FIntPoint, EAddError> Resolution = GetCameraResolution(InCamera);
+
+	if (Resolution.HasValue())
+	{
+		if (CommonResolution == FIntPoint::NoneValue)
+		{
+			CommonResolution = Resolution.GetValue();
+		}
+
+		if (CommonResolution != FIntPoint::NoneValue)
+		{
+			bAllEqual &= Resolution.GetValue() == CommonResolution;
+		}
+		else
+		{
+			bAllEqual = false;
+		}
+	}
+
+	return Resolution;
+}
+
+TValueOrError<FIntPoint, FResolutionResolver::EResolveError> FResolutionResolver::Resolve() const
+{
+	if (bAllEqual)
+	{
+		return MakeValue(CommonResolution);
+	}
+	else
+	{
+		return MakeError(EResolveError::Mismatched);
+	}
+}
+
+TValueOrError<FIntPoint, FResolutionResolver::EAddError> FResolutionResolver::GetCameraResolution(const FStereoCameraTakeInfo::FCamera& InCamera) const
+{
+	if (InCamera.Resolution == FIntPoint::NoneValue)
+	{
+		if (FPaths::DirectoryExists(InCamera.FramesPath))
+		{
+			return GetResolutionFromSingleImage(InCamera.FramesPath);
+		}
+		else
+		{
+			return MakeError(EAddError::FramesPathDoesNotExist);
+		}
+	}
+	else
+	{
+		// Property has already been set on the camera so just use it. 
+		// 
+		// This is a little weird, as it's effectively passing responsibility for determining the resolution value out of this function/class. 
+		// However, this functionality has been brought across from the previous implementation and might need some more consideration to change.
+
+		return MakeValue(InCamera.Resolution);
+	}
+}
+
+TValueOrError<FIntPoint, FResolutionResolver::EAddError> FResolutionResolver::GetResolutionFromSingleImage(const FString& InDirectoryPath) const
+{
+	FString FirstFilePath;
+	TArray<FString> SupportedExtensions = { TEXT("png"), TEXT("jpg"), TEXT("jpeg") };
+
+	IFileManager::Get().IterateDirectory(*InDirectoryPath, [&FirstFilePath, &SupportedExtensions](const TCHAR* InFilenameOrDirectory, bool bInIsDirectory) mutable
+	{
+		if (!bInIsDirectory &&
+			SupportedExtensions.Contains(FPaths::GetExtension(InFilenameOrDirectory)))
+		{
+			FString FileNamePath = InFilenameOrDirectory;
+			FPaths::NormalizeFilename(FileNamePath);
+			FirstFilePath = MoveTemp(FileNamePath);
+
+			// Returning false because we only need the first file
+			return false;
+		}
+
+		return true;
+	});
+
+	if (FirstFilePath.IsEmpty())
+	{
+		return MakeError(EAddError::NoImagesFound);
+	}
+
+	TArray<uint8> RawFileData;
+	if (!FFileHelper::LoadFileToArray(RawFileData, *FirstFilePath))
+	{
+		return MakeError(EAddError::ImageLoadFailed);
+	}
+
+	IImageWrapperModule& ImageWrapperModule = FModuleManager::LoadModuleChecked<IImageWrapperModule>(TEXT("ImageWrapper"));
+	EImageFormat ImageFormat = ImageWrapperModule.DetectImageFormat(RawFileData.GetData(), RawFileData.Num());
+	TSharedPtr<IImageWrapper> ImageWrapper = ImageWrapperModule.CreateImageWrapper(ImageFormat);
+
+	if (!ImageWrapper.IsValid() || !ImageWrapper->SetCompressed(RawFileData.GetData(), RawFileData.Num()))
+	{
+		return MakeError(EAddError::InvalidImageWrapper);
+	}
+
+	return MakeValue(static_cast<int>(ImageWrapper->GetWidth()), static_cast<int>(ImageWrapper->GetHeight()));
+}
+
+}
diff --git a/Engine/Plugins/VirtualProduction/CaptureManager/CaptureManagerCore/Source/CaptureMetadataExtraction/Private/ResolutionResolver.h b/Engine/Plugins/VirtualProduction/CaptureManager/CaptureManagerCore/Source/CaptureMetadataExtraction/Private/ResolutionResolver.h
new file mode 100644
--- /dev/null
+++ b/Engine/Plugins/VirtualProduction/CaptureManager/CaptureManagerCore/Source/CaptureMetadataExtraction/Private/ResolutionResolver.h
@@ -0,0 +1,47 @@
+// Copyright Epic Games, Inc. All Rights Reserved.
+
+#pragma once
+
+#include "StereoCameraTakeMetadata.h"
+
+#include "Containers/UnrealString.h"
+#include "Math/IntPoint.h"
+#include "Misc/Optional.h"
+
+namespace UE::CaptureManager
+{
+
+/**
+ * Class used to determine and validate the image resolution for a set of stereo camera data.
+ *
+ * Each time a camera is added to the resolver, the internal state is updated and will influence the result of the final Resolve() call. A collective
+ * value for the image resolution will only be returned if all of the cameras share the same image resolution.
+ */
+class FResolutionResolver
+{
+public:
+	enum class EAddError
+	{
+		FramesPathDoesNotExist,
+		NoImagesFound,
+		ImageLoadFailed,
+		InvalidImageWrapper
+	};
+
+	enum class EResolveError
+	{
+		Mismatched
+	};
+
+	TValueOrError<FIntPoint, EAddError> Add(const FStereoCameraTakeInfo::FCamera& Camera);
+	TValueOrError<FIntPoint, EResolveError> Resolve() const;
+
+private:
+	TValueOrError<FIntPoint, EAddError> GetCameraResolution(const FStereoCameraTakeInfo::FCamera& Camera) const;
+	TValueOrError<FIntPoint, EAddError> GetResolutionFromSingleImage(const FString& InDirectoryPath) const;
+
+	FIntPoint CommonResolution = FIntPoint::NoneValue;
+	bool bAllEqual = true;
+};
+
+}
\ No newline at end of file
diff --git a/Engine/Plugins/VirtualProduction/CaptureManager/CaptureManagerCore/Source/CaptureMetadataExtraction/Private/StereoCameraMetadataParseUtils.cpp b/Engine/Plugins/VirtualProduction/CaptureManager/CaptureManagerCore/Source/CaptureMetadataExtraction/Private/StereoCameraMetadataParseUtils.cpp
new file mode 100644
--- /dev/null
+++ b/Engine/Plugins/VirtualProduction/CaptureManager/CaptureManagerCore/Source/CaptureMetadataExtraction/Private/StereoCameraMetadataParseUtils.cpp
@@ -0,0 +1,205 @@
+// Copyright Epic Games, Inc. All Rights Reserved.
+
+#include "StereoCameraMetadataParseUtils.h"
+
+#include "StereoCameraTakeMetadata.h"
+
+#include "Dom/JsonObject.h"
+#include "Dom/JsonValue.h"
+#include "Misc/FileHelper.h"
+#include "Misc/Paths.h"
+#include "Serialization/JsonSerializer.h"
+
+#include "Modules/ModuleManager.h"
+#include "IImageWrapperModule.h"
+
+#define LOCTEXT_NAMESPACE "StereoCameraTakeMetadata"
+
+TOptional<FText> TakeDurationExceedsLimit(const float InDurationInSeconds)
+{
+	IConsoleVariable* CVar = IConsoleManager::Get().FindConsoleVariable(TEXT("au.SoundWaveImportLengthLimitInSeconds"));
+	if (!CVar)
+	{
+		return {};
+	}
+
+	static constexpr float Unlimited = -1.f;
+	float Limit = CVar->GetFloat();
+
+	if (FMath::IsNearlyEqual(Limit, Unlimited) && FMath::IsNegativeOrNegativeZero(Limit - InDurationInSeconds))
+	{
+		return {};
+	}
+
+	const FText Message = LOCTEXT("TakeDurationExceedsLimit", "Take duration ({0} seconds) exceeds allowed limit ({1} seconds).");
+
+	FNumberFormattingOptions Options;
+	Options.MaximumFractionalDigits = 2;
+	Options.MinimumFractionalDigits = 2;
+
+	return FText::Format(Message, FText::AsNumber(InDurationInSeconds, &Options), FText::AsNumber(Limit, &Options));
+}
+
+FString DetermineImageFormat(const FString& InImageFolderPath)
+{
+	FString Extension;
+	IImageWrapperModule& ImageWrapperModule = FModuleManager::LoadModuleChecked<IImageWrapperModule>(TEXT("ImageWrapper"));
+
+	IFileManager::Get().IterateDirectory(*InImageFolderPath, [&Extension, &ImageWrapperModule](const TCHAR* InFileName, bool bIsDirectory)
+		{
+			if (bIsDirectory)
+			{
+				return true;
+			}
+
+			Extension = FPaths::GetExtension(InFileName);
+			EImageFormat Format = ImageWrapperModule.GetImageFormatFromExtension(*Extension);
+
+			if (Format != EImageFormat::Invalid)
+			{
+				return true;
+			}
+
+			return false; // Returns false because we are looking only for the first image in the given directory
+		});
+
+	return Extension;
+}
+
+FTakeMetadata ConvertOldToNewTakeMetadata(const FString& InTakeFolder, const FStereoCameraTakeInfo& InStereoCameraInfo)
+{
+	FString CalibrationPath = InTakeFolder / TEXT("calib.json");
+
+	FTakeMetadata::FCalibration Calibration;
+	Calibration.Name = TEXT("undefined");
+	Calibration.Format = TEXT("opencv");
+	Calibration.Path = InStereoCameraInfo.CalibrationFilePath;
+
+	FTakeMetadata::FDevice::FPlatform Platform;
+	Platform.Name = TEXT("Windows");
+
+	FTakeMetadata Metadata;
+	Metadata.Device.Platform = MoveTemp(Platform);
+
+	Metadata.Calibration.Push(MoveTemp(Calibration));
+	Metadata.Device.Type = InStereoCameraInfo.DeviceInfo.Type;
+	Metadata.Device.Model = InStereoCameraInfo.DeviceInfo.Model;
+	Metadata.Device.Name = InStereoCameraInfo.DeviceInfo.Id;
+
+	Metadata.UniqueId = InStereoCameraInfo.Id;
+	Metadata.Slate = InStereoCameraInfo.Slate;
+	Metadata.TakeNumber = InStereoCameraInfo.Take;
+	Metadata.DateTime = InStereoCameraInfo.Date;
+
+	Metadata.Thumbnail = FTakeThumbnailData(InStereoCameraInfo.ThumbnailPath);
+
+	Metadata.Version.Major = 3;
+	Metadata.Version.Minor = 0;
+
+	TOptional<float> AudioDuration = 0.0f;
+
+	TArray<FTakeMetadata::FVideo> VideoArray;
+	for (const TPair<FString, FStereoCameraTakeInfo::FCamera>& Camera : InStereoCameraInfo.CameraMap)
+	{
+		FTakeMetadata::FVideo Video;
+		Video.Name = Camera.Value.UserId;
+		Video.Path = Camera.Value.FramesPath;
+		Video.PathType = FTakeMetadata::FVideo::EPathType::Folder;
+		Video.Format = DetermineImageFormat(Camera.Value.FramesPath);
+		Video.Orientation = FTakeMetadata::FVideo::EOrientation::Original;
+		Video.FramesCount = Camera.Value.FrameRange.Value - Camera.Value.FrameRange.Key + 1;
+
+		Video.FrameHeight = Camera.Value.Resolution.Y;
+		Video.FrameWidth = Camera.Value.Resolution.X;
+
+		Video.FrameRate = Camera.Value.FrameRate;
+
+		Video.TimecodeStart = Camera.Value.StartTimecode;
+
+		for (const FFrameRange& FrameRange : Camera.Value.CaptureExcludedFrames)
+		{
+			TArray<uint32> DroppedFrames;
+			for (int32 Frame = FrameRange.StartFrame; Frame <= FrameRange.EndFrame; ++Frame)
+			{
+				DroppedFrames.Add(Frame);
+			}
+			Video.DroppedFrames = MoveTemp(DroppedFrames);
+		}
+
+		if (Video.FramesCount.IsSet())
+		{
+			// We do not know the audio duration and so we must estimate it from the video. This apporach has been deemed 
+			// acceptable for the moment, based on how this duration value gets used.
+			// If the video duration has not been set, then we cannot estimate the audio duration
+			const float VideoDuration = static_cast<float>(Video.FramesCount.GetValue()) / Video.FrameRate;
+			AudioDuration = FMath::Max(0, VideoDuration);
+		}
+
+		VideoArray.Add(MoveTemp(Video));
+	}
+
+	Metadata.Video = MoveTemp(VideoArray);
+
+	TArray<FTakeMetadata::FAudio> AudioArray;
+	AudioArray.Reserve(InStereoCameraInfo.AudioArray.Num());
+
+	for (const FStereoCameraTakeInfo::FAudio& AudioEntry : InStereoCameraInfo.AudioArray)
+	{
+		FTakeMetadata::FAudio Audio;
+		Audio.Duration = AudioDuration;
+		Audio.Path = AudioEntry.StreamPath;
+		Audio.TimecodeStart = AudioEntry.StartTimecode;
+		Audio.TimecodeRate = AudioEntry.TimecodeRate;
+
+		AudioArray.Add(MoveTemp(Audio));
+	}
+
+	Metadata.Audio = MoveTemp(AudioArray);
+
+	return Metadata;
+}
+
+namespace UE::CaptureManager::StereoCameraMetadata
+{
+
+TOptional<FTakeMetadata> ParseOldStereoCameraMetadata(const FString& InTakeFolder, TArray<FText>& OutValidationError)
+{
+	const FString MetadataFilePath = InTakeFolder / TEXT("take.json");
+	TOptional<FStereoCameraTakeInfo> StereoCameraInfoResult = FStereoCameraSystemTakeParser::ParseTakeMetadataFile(MetadataFilePath);
+
+	if (!StereoCameraInfoResult.IsSet())
+	{
+		OutValidationError.Add(LOCTEXT("ParseOldStereoCameraTakeMetadata_ParseTakeInfoFailed", "Failed to parse old take json file"));
+		return {};
+	}
+
+	FStereoCameraTakeInfo StereoCameraInfo = StereoCameraInfoResult.GetValue();
+
+	TArray<FText> Issues = FStereoCameraSystemTakeParser::CheckStereoCameraTakeInfo(InTakeFolder, StereoCameraInfo, 2, TEXT("HMC"));
+
+	if (!Issues.IsEmpty())
+	{
+		OutValidationError.Append(Issues);
+	}
+
+	TArray<FText> ResolveIssues = FStereoCameraSystemTakeParser::ResolveResolution(StereoCameraInfo);
+	if (!ResolveIssues.IsEmpty())
+	{
+		OutValidationError.Append(ResolveIssues);
+	}
+
+	FStereoCameraTakeInfo::FCamera CameraInfo = StereoCameraInfo.CameraMap.CreateConstIterator()->Value;
+	uint32 NumOfTakes = (CameraInfo.FrameRange.Value - CameraInfo.FrameRange.Key) + 1;
+
+	TOptional<FText> DurationLimit = TakeDurationExceedsLimit(static_cast<float>(NumOfTakes) / CameraInfo.FrameRate);
+	if (DurationLimit.IsSet())
+	{
+		OutValidationError.Add(DurationLimit.GetValue());
+	}
+
+	return ConvertOldToNewTakeMetadata(InTakeFolder, StereoCameraInfo);
+}
+
+}
+
+#undef LOCTEXT_NAMESPACE
diff --git a/Engine/Plugins/VirtualProduction/CaptureManager/CaptureManagerCore/Source/CaptureMetadataExtraction/Private/StereoCameraMetadataParseUtils.h b/Engine/Plugins/VirtualProduction/CaptureManager/CaptureManagerCore/Source/CaptureMetadataExtraction/Private/StereoCameraMetadataParseUtils.h
new file mode 100644
--- /dev/null
+++ b/Engine/Plugins/VirtualProduction/CaptureManager/CaptureManagerCore/Source/CaptureMetadataExtraction/Private/StereoCameraMetadataParseUtils.h
@@ -0,0 +1,16 @@
+// Copyright Epic Games, Inc. All Rights Reserved.
+
+#pragma once
+
+#include "CaptureManagerTakeMetadata.h"
+
+#define UE_API CAPTUREMETADATAEXTRACTION_API
+
+namespace UE::CaptureManager::StereoCameraMetadata
+{
+
+UE_API TOptional<FTakeMetadata> ParseOldStereoCameraMetadata(const FString& InTakeFolder, TArray<FText>& OutValidationError);
+
+}
+
+#undef UE_API
diff --git a/Engine/Plugins/VirtualProduction/CaptureManager/CaptureManagerCore/Source/CaptureMetadataExtraction/Private/StereoCameraTakeMetadata.cpp b/Engine/Plugins/VirtualProduction/CaptureManager/CaptureManagerCore/Source/CaptureMetadataExtraction/Private/StereoCameraTakeMetadata.cpp
new file mode 100644
--- /dev/null
+++ b/Engine/Plugins/VirtualProduction/CaptureManager/CaptureManagerCore/Source/CaptureMetadataExtraction/Private/StereoCameraTakeMetadata.cpp
@@ -0,0 +1,673 @@
+// Copyright Epic Games, Inc. All Rights Reserved.
+
+#include "StereoCameraTakeMetadata.h"
+
+#include "ResolutionResolver.h"
+
+#include "Dom/JsonObject.h"
+#include "Dom/JsonValue.h"
+#include "Misc/FileHelper.h"
+#include "Misc/Paths.h"
+#include "Serialization/JsonSerializer.h"
+
+DEFINE_LOG_CATEGORY_STATIC(LogStereoCameraMetadata, Log, All);
+
+#define LOCTEXT_NAMESPACE "StereoCameraTakeMetadata"
+
+#define CHECK_RESULT_AND_RETURN(Fun) if (bool Result = Fun; !Result) {return TOptional<FStereoCameraTakeInfo>{};}
+
+static const FText FailedToDetermineCameraResolution = LOCTEXT("FailedToDetermineCameraResolution", "Failed to determine camera resolution");
+static const FText NoImagesFound = LOCTEXT("NoImagesFound", "No images found");
+static const FText ImageLoadFailed = LOCTEXT("ImageLoadFailed", "Failed to load an image");
+static const FText CameraResolutionMismatch = LOCTEXT("CameraResolutionMismatch", "Camera resolutions do not match");
+
+static FString ConvertPathToAbsolute(const FString& InBasePath, const FString& InPath)
+{
+	// Convert paths to absolute
+	if (FPaths::IsRelative(InPath))
+	{
+		const FString Directory = FPaths::GetPath(InBasePath);
+		return FPaths::Combine(Directory, InPath);
+	}
+
+	return InPath;
+}
+
+static FString WrapLogMessage(const FString InMessage, const FStereoCameraTakeInfo& InTakeInfo)
+{
+	return FString::Format(TEXT("{0} ({1} #{2})"), { InMessage, InTakeInfo.Slate, InTakeInfo.Take });
+}
+
+static FString WrapLogMessage(const FString InMessage, const FTakeMetadata& InTakeInfo)
+{
+	return FString::Format(TEXT("{0} ({1} #{2})"), { InMessage, InTakeInfo.Slate, InTakeInfo.TakeNumber });
+}
+
+template<typename FTakeInfoType>
+static void ReportIssue(FText InMessage, const FTakeInfoType& InTakeInfo, TArray<FText>& OutIssues)
+{
+	// Log and register as an issue, it's important to do both so the user has a reference log that can be sent, rather than just screenshots.
+	UE_LOGF(LogStereoCameraMetadata, Warning, "%ls", *WrapLogMessage(InMessage.ToString(), InTakeInfo));
+	OutIssues.Emplace(MoveTemp(InMessage));
+}
+
+static void ReportResolutionAddError(
+	const UE::CaptureManager::FResolutionResolver::EAddError AddError,
+	const FString& CameraId,
+	const FStereoCameraTakeInfo& InStereoCameraTakeInfo,
+	TArray<FText>& OutIssues
+)
+{
+	using namespace UE::CaptureManager;
+
+	FText Message;
+
+	switch (AddError)
+	{
+		case FResolutionResolver::EAddError::FramesPathDoesNotExist:
+		{
+			// Manually format the message in this case, as we need a little more control
+			const FStereoCameraTakeInfo::FCamera* Camera = InStereoCameraTakeInfo.CameraMap.Find(CameraId);
+
+			if (Camera)
+			{
+				// We can provide more information in this case and we want to put the "for camera" part before the path, which can be very long
+				ReportIssue(
+					FText::Format(
+						LOCTEXT("FramesPathDoesNotExistForCamera", "Frames path does not exist for camera: {0} ({1})"),
+						FText::FromString(CameraId),
+						FText::FromString(Camera->FramesPath)
+					),
+					InStereoCameraTakeInfo,
+					OutIssues
+				);
+			}
+			else
+			{
+				ReportIssue(
+					FText::Format(
+						LOCTEXT("FramesPathDoesNotExist", "Frames path does not exist: {0}"),
+						FText::FromString(CameraId)
+					),
+					InStereoCameraTakeInfo,
+					OutIssues
+				);
+			}
+			break;
+		}
+
+		case FResolutionResolver::EAddError::NoImagesFound:
+			Message = NoImagesFound;
+			break;
+
+		case FResolutionResolver::EAddError::ImageLoadFailed:
+			Message = ImageLoadFailed;
+			break;
+
+		case FResolutionResolver::EAddError::InvalidImageWrapper:
+			// Log detailed message but don't display that to the user (they don't know what the image wrapper is)
+			UE_LOGF(LogStereoCameraMetadata, Warning, "%ls",
+				   *WrapLogMessage(FString::Format(TEXT("Image wrapper is in an invalid state for camera: {0}"), { CameraId, }), InStereoCameraTakeInfo)
+			);
+			Message = FailedToDetermineCameraResolution;
+			break;
+
+		default:
+			Message = FailedToDetermineCameraResolution;
+			break;
+	}
+
+	if (!Message.IsEmpty())
+	{
+		FText IssueMessage = FText::Format(LOCTEXT("ResolutionAddError", "{0} for camera: {1}"), Message, FText::FromString(CameraId));
+		ReportIssue(MoveTemp(IssueMessage), InStereoCameraTakeInfo, OutIssues);
+	}
+}
+
+static void ReportResolutionResolveError(
+	const UE::CaptureManager::FResolutionResolver::EResolveError ResolutionError,
+	const FStereoCameraTakeInfo& InStereoCameraTakeInfo,
+	TArray<FText>& OutIssues
+)
+{
+	using namespace UE::CaptureManager;
+
+	FText Message;
+
+	switch (ResolutionError)
+	{
+		case FResolutionResolver::EResolveError::Mismatched:
+			Message = CameraResolutionMismatch;
+			break;
+		default:
+			Message = FailedToDetermineCameraResolution;
+			break;
+	}
+
+	ReportIssue(MoveTemp(Message), InStereoCameraTakeInfo, OutIssues);
+}
+
+static TArray<TSharedPtr<FJsonValue>> ParseJsonArrayFromFile(const FString& InFilePath)
+{
+	FString JsonStringBuffer;
+	TArray<TSharedPtr<FJsonValue>> Result;
+
+	if (FFileHelper::LoadFileToString(JsonStringBuffer, *InFilePath) &&
+		FJsonSerializer::Deserialize(TJsonReaderFactory<>::Create(JsonStringBuffer), Result))
+	{
+		return Result;
+	}
+
+	return TArray<TSharedPtr<FJsonValue>>();
+}
+
+static TSharedPtr<FJsonObject> ParseJsonObjectFromFile(const FString& InFilePath)
+{
+	FString JsonStringBuffer;
+
+	if (!IFileManager::Get().FileExists(*InFilePath))
+	{
+		UE_LOGF(LogStereoCameraMetadata, Error, "File not found: %ls", *InFilePath);
+		return nullptr;
+	}
+
+	if (!FFileHelper::LoadFileToString(JsonStringBuffer, *InFilePath))
+	{
+		UE_LOGF(LogStereoCameraMetadata, Error, "Failed to load file (check permissions): %ls", *InFilePath);
+		return nullptr;
+	}
+
+	TSharedPtr<FJsonObject> Result;
+
+	if (!FJsonSerializer::Deserialize(TJsonReaderFactory<>::Create(JsonStringBuffer), Result))
+	{
+		UE_LOGF(LogStereoCameraMetadata, Error, "Failed to load json file (check for syntax errors): %ls", *InFilePath);
+		return nullptr;
+	}
+
+	return Result;
+}
+
+static void CheckCameraIds(const FTakeMetadata& InStereoCameraTakeInfo, const TMap<FString, FStereoCameraInfo>& OutTakeCameras, TArray<FText>& OutIssues)
+{
+	for (const FTakeMetadata::FVideo& ExpectedCamera : InStereoCameraTakeInfo.Video)
+	{
+		const FString& ExpectedId = ExpectedCamera.Name;
+
+		if (!OutTakeCameras.Contains(ExpectedId))
+		{
+			ReportIssue(
+				FText::Format(
+					LOCTEXT("CameraUserIdMismatch", "Camera ID present in take metadata but not present in the calibration: {0}"),
+					FText::FromString(ExpectedId)
+				),
+				InStereoCameraTakeInfo,
+				OutIssues
+			);
+		}
+	}
+}
+
+static void CheckAudio(const FStereoCameraTakeInfo& InStereoCameraTakeInfo, TArray<FText>& OutIssues)
+{
+	for (const FStereoCameraTakeInfo::FAudio& Audio : InStereoCameraTakeInfo.AudioArray)
+	{
+		if (!FPaths::FileExists(Audio.StreamPath))
+		{
+			ReportIssue(
+				FText::Format(LOCTEXT("MissingAudio", "Audio file not found: {0}"), FText::FromString(Audio.StreamPath)),
+				InStereoCameraTakeInfo,
+				OutIssues
+			);
+		}
+	}
+}
+
+static bool ParseCalibrationFile(const FString& InFileName, TMap<FString, FStereoCameraInfo>& OutCameras)
+{
+	const TArray<TSharedPtr<FJsonValue>> CalibrationJson = ParseJsonArrayFromFile(InFileName);
+
+	if (CalibrationJson.IsEmpty())
+	{
+		UE_LOGF(LogStereoCameraMetadata, Warning, "Calibration file is empty: %ls", *InFileName);
+		return false;
+	}
+
+	for (const TSharedPtr<FJsonValue>& Value : CalibrationJson)
+	{
+		const TSharedPtr<FJsonObject>* MaybeObject = nullptr;
+		if (Value->TryGetObject(MaybeObject))
+		{
+			const TSharedRef<FJsonObject> Object = MaybeObject->ToSharedRef();
+
+			const TSharedPtr<FJsonObject>* MaybeMetaData = nullptr;
+			if (Object->TryGetObjectField(TEXT("metadata"), MaybeMetaData))
+			{
+				const TSharedRef<FJsonObject> MetaData = MaybeMetaData->ToSharedRef();
+
+				FStereoCameraInfo CameraInfo;
+				MetaData->TryGetStringField(TEXT("camera"), CameraInfo.Name);
+
+				if (!CameraInfo.Name.IsEmpty())
+				{
+					bool bOK = true;
+					CameraInfo.Calibration.CameraId = CameraInfo.Name;
+					bOK &= Object->TryGetNumberField(TEXT("image_size_x"), CameraInfo.Calibration.ImageSize.X);
+					bOK &= Object->TryGetNumberField(TEXT("image_size_y"), CameraInfo.Calibration.ImageSize.Y);
+					bOK &= Object->TryGetNumberField(TEXT("fx"), CameraInfo.Calibration.FocalLength.X);
+					bOK &= Object->TryGetNumberField(TEXT("fy"), CameraInfo.Calibration.FocalLength.Y);
+					bOK &= Object->TryGetNumberField(TEXT("cx"), CameraInfo.Calibration.PrincipalPoint.X);
+					bOK &= Object->TryGetNumberField(TEXT("cy"), CameraInfo.Calibration.PrincipalPoint.Y);
+					bOK &= Object->TryGetNumberField(TEXT("k1"), CameraInfo.Calibration.K1);
+					bOK &= Object->TryGetNumberField(TEXT("k2"), CameraInfo.Calibration.K2);
+					bOK &= Object->TryGetNumberField(TEXT("k3"), CameraInfo.Calibration.K3);
+					bOK &= Object->TryGetNumberField(TEXT("p1"), CameraInfo.Calibration.P1);
+					bOK &= Object->TryGetNumberField(TEXT("p2"), CameraInfo.Calibration.P2);
+
+					const TArray<TSharedPtr<FJsonValue>>* Transform = nullptr;
+					bOK &= Object->TryGetArrayField(TEXT("transform"), Transform);
+					bOK &= Transform && Transform->Num() == 16;
+
+					if (bOK)
+					{
+						for (int32 I = 0; I < 4; ++I)
+						{
+							for (int32 J = 0; J < 4; ++J)
+							{
+								CameraInfo.Calibration.Transform.M[J][I] = (*Transform)[I * 4 + J]->AsNumber();
+							}
+						}
+
+						FString Name = CameraInfo.Name;
+						OutCameras.Add(MoveTemp(Name), MoveTemp(CameraInfo));
+					}
+				}
+			}
+		}
+	}
+
+	return true;
+}
+
+static void LoadCalibration(const FTakeMetadata& InStereoCameraTakeInfo, TMap<FString, FStereoCameraInfo>& OutTakeCameras, TArray<FText>& OutIssues)
+{
+	if (InStereoCameraTakeInfo.Calibration.IsEmpty())
+	{
+		ReportIssue(
+			LOCTEXT("CalibrationNotFound", "Calibration object not set"),
+			InStereoCameraTakeInfo,
+			OutIssues
+		);
+
+		return;
+	}
+
+	FTakeMetadata::FCalibration Calibration = InStereoCameraTakeInfo.Calibration[0];
+
+	if (FPaths::FileExists(Calibration.Path))
+	{
+		if (!ParseCalibrationFile(Calibration.Path, OutTakeCameras))
+		{
+			ReportIssue(
+				FText::Format(
+					LOCTEXT("CalibrationLoadFailed", "Failed to load calibration: {0}"),
+					FText::FromString(Calibration.Path)
+				),
+				InStereoCameraTakeInfo,
+				OutIssues
+			);
+		}
+	}
+	else
+	{
+		ReportIssue(
+			FText::Format(LOCTEXT("CalibrationFileNotFound", "Calibration file not found: {0}"), FText::FromString(Calibration.Path)),
+			InStereoCameraTakeInfo,
+			OutIssues
+		);
+	}
+}
+
+static void CheckForEmptyFramesPaths(const FStereoCameraTakeInfo& InStereoCameraTakeInfo, TArray<FText>& OutIssues)
+{
+	// Work-around for lack of enforcement in take parsing code (which allows this to occur)
+
+	for (const TPair<FString, FStereoCameraTakeInfo::FCamera>& Camera : InStereoCameraTakeInfo.CameraMap)
+	{
+		if (Camera.Value.FramesPath.IsEmpty())
+		{
+			ReportIssue(
+				FText::Format(LOCTEXT("EmptyFramesPath", "Empty frames path for camera: {0}"), FText::FromString(Camera.Value.UserId)),
+				InStereoCameraTakeInfo,
+				OutIssues
+			);
+		}
+	}
+}
+
+static void CheckForNonANSICharacters(const FStereoCameraTakeInfo& InStereoCameraTakeInfo, TArray<FText>& OutIssues)
+{
+	TArray<FString> NonANSIContainingProperties;
+
+	if (!FCString::IsPureAnsi(*InStereoCameraTakeInfo.Id))
+	{
+		NonANSIContainingProperties.Add(TEXT("Id"));
+	}
+
+	if (!FCString::IsPureAnsi(*InStereoCameraTakeInfo.Slate))
+	{
+		NonANSIContainingProperties.Add(TEXT("Slate"));
+	}
+
+	if (!FCString::IsPureAnsi(*InStereoCameraTakeInfo.ThumbnailPath))
+	{
+		NonANSIContainingProperties.Add(TEXT("ThumbnailPath"));
+	}
+
+	if (!FCString::IsPureAnsi(*InStereoCameraTakeInfo.CalibrationFilePath))
+	{
+		NonANSIContainingProperties.Add(TEXT("CalibrationFilePath"));
+	}
+
+	for (const TPair<FString, FStereoCameraTakeInfo::FCamera>& Camera : InStereoCameraTakeInfo.CameraMap)
+	{
+		const FString& CameraUserId = Camera.Value.UserId;
+		FString CameraPropertyPrefix = FString::Format(TEXT("Camera({0})."), { CameraUserId });
+		if (!FCString::IsPureAnsi(*CameraUserId))
+		{
+			NonANSIContainingProperties.Add(CameraPropertyPrefix + TEXT("UserId"));
+		}
+
+		if (!FCString::IsPureAnsi(*Camera.Value.FramesPath))
+		{
+			NonANSIContainingProperties.Add(CameraPropertyPrefix + TEXT("FramesPath"));
+		}
+	}
+
+	for (const FStereoCameraTakeInfo::FAudio& Audio : InStereoCameraTakeInfo.AudioArray)
+	{
+		const FString& AudioUserId = Audio.UserId;
+		FString AudioPropertyPrefix = FString::Format(TEXT("Audio({0})."), { AudioUserId });
+		if (!FCString::IsPureAnsi(*AudioUserId))
+		{
+			NonANSIContainingProperties.Add(AudioPropertyPrefix + TEXT("UserId"));
+		}
+
+		if (!FCString::IsPureAnsi(*Audio.StreamPath))
+		{
+			NonANSIContainingProperties.Add(AudioPropertyPrefix + TEXT("StreamPath"));
+		}
+	}
+
+	for (const FString& NonANSIContainingProperty : NonANSIContainingProperties)
+	{
+		const FText Message = FText::Format(LOCTEXT("UnsupportedCharactersWithinTakeInfo", "Take '{0}' contains unsupported text characters"), FText::FromString(NonANSIContainingProperty));
+		OutIssues.Add(Message);
+		UE_LOGF(LogStereoCameraMetadata, Warning, "InStereoCameraTakeInfo property '%ls' contains unsupported non-ansi text characters.", *NonANSIContainingProperty);
+	}
+}
+
+FString FStereoCameraTakeInfo::GetName() const
+{
+	return FString::Format(TEXT("{0}_{1}"), { Slate, Take });
+}
+
+FString FStereoCameraTakeInfo::GetFolderName() const
+{
+	return FPaths::GetPathLeaf(FPaths::GetPath(TakeJsonFilePath));
+}
+
+TOptional<FStereoCameraTakeInfo> FStereoCameraSystemTakeParser::ParseTakeMetadataFile(const FString& InFileName)
+{
+	FStereoCameraTakeInfo CubicTakeInfo;
+	const FString CurrentFileName = FPaths::GetCleanFilename(InFileName);
+
+	if (CurrentFileName == TEXT("take.json"))
+	{
+		if (TSharedPtr<FJsonObject> TakeMetadataJson = ParseJsonObjectFromFile(InFileName))
+		{
+			CHECK_RESULT_AND_RETURN(TakeMetadataJson->TryGetNumberField(TEXT("Version"), CubicTakeInfo.Version));
+			CHECK_RESULT_AND_RETURN(TakeMetadataJson->TryGetStringField(TEXT("Id"), CubicTakeInfo.Id));
+			CHECK_RESULT_AND_RETURN(TakeMetadataJson->TryGetNumberField(TEXT("Take"), CubicTakeInfo.Take));
+			CHECK_RESULT_AND_RETURN(TakeMetadataJson->TryGetStringField(TEXT("Slate"), CubicTakeInfo.Slate));
+			// Optional
+			FString Thumbnail;
+			if (TakeMetadataJson->TryGetStringField(TEXT("Thumbnail"), Thumbnail))
+			{
+				CubicTakeInfo.ThumbnailPath = ConvertPathToAbsolute(InFileName, Thumbnail);
+			}
+
+			FString Date;
+			CHECK_RESULT_AND_RETURN(TakeMetadataJson->TryGetStringField(TEXT("LocalDateTime"), Date));
+
+			if (!FDateTime::ParseIso8601(*Date, CubicTakeInfo.Date))
+			{
+				UE_LOGF(LogStereoCameraMetadata, Warning, "Failed to parse date string '%ls' as an ISO8601 date", *Date);
+			}
+
+			FString CalibrationFilePath;
+			CHECK_RESULT_AND_RETURN(TakeMetadataJson->TryGetStringField(TEXT("CalibrationInfo"), CalibrationFilePath));
+			CubicTakeInfo.CalibrationFilePath = ConvertPathToAbsolute(InFileName, CalibrationFilePath);
+
+			const TSharedPtr<FJsonObject>* DeviceInfoObject;
+			CHECK_RESULT_AND_RETURN(TakeMetadataJson->TryGetObjectField(TEXT("DeviceInfo"), DeviceInfoObject));
+
+			CHECK_RESULT_AND_RETURN((*DeviceInfoObject)->TryGetStringField(TEXT("Model"), CubicTakeInfo.DeviceInfo.Model));
+			CHECK_RESULT_AND_RETURN((*DeviceInfoObject)->TryGetStringField(TEXT("Type"), CubicTakeInfo.DeviceInfo.Type));
+			CHECK_RESULT_AND_RETURN((*DeviceInfoObject)->TryGetStringField(TEXT("Id"), CubicTakeInfo.DeviceInfo.Id));
+
+			const TArray<TSharedPtr<FJsonValue>>* CamerasJson;
+			CHECK_RESULT_AND_RETURN(TakeMetadataJson->TryGetArrayField(TEXT("Cameras"), CamerasJson));
+
+			// Cameras
+			for (const TSharedPtr<FJsonValue>& CameraJson : *CamerasJson)
+			{
+				FStereoCameraTakeInfo::FCamera Camera;
+				const TSharedPtr<FJsonObject>& CameraObject = CameraJson->AsObject();
+
+				CHECK_RESULT_AND_RETURN(CameraObject->TryGetStringField(TEXT("UserID"), Camera.UserId));
+
+				// Optional
+				const TArray<TSharedPtr<FJsonValue>>* FrameRangeJson = nullptr;
+				if (CameraObject->TryGetArrayField(TEXT("FrameRange"), FrameRangeJson))
+				{
+					CHECK_RESULT_AND_RETURN(((*FrameRangeJson).Num() == 2));
+					Camera.FrameRange.Key = (*FrameRangeJson)[0]->AsNumber();
+					Camera.FrameRange.Value = (*FrameRangeJson)[1]->AsNumber();
+				}
+
+				CHECK_RESULT_AND_RETURN(CameraObject->TryGetNumberField(TEXT("FrameRate"), Camera.FrameRate));
+
+				FString FramesPath;
+				CHECK_RESULT_AND_RETURN(CameraObject->TryGetStringField(TEXT("FramesPath"), FramesPath));
+
+				if (!FramesPath.IsEmpty())
+				{
+					Camera.FramesPath = ConvertPathToAbsolute(InFileName, FramesPath);
+				}
+
+				// Optional
+				CameraObject->TryGetStringField(TEXT("StartTimecode"), Camera.StartTimecode);
+
+				// Optional
+				const TArray<TSharedPtr<FJsonValue>>* FramesDroppedJson;
+				if (CameraObject->TryGetArrayField(TEXT("FramesDropped"), FramesDroppedJson))
+				{
+					for (const TSharedPtr<FJsonValue>& DroppedFrameJson : *FramesDroppedJson)
+					{
+						FString DroppedFrame;
+						CHECK_RESULT_AND_RETURN(DroppedFrameJson->TryGetString(DroppedFrame));
+
+						TArray<FString> FrameTokens;
+						DroppedFrame.ParseIntoArray(FrameTokens, TEXT("-"));
+						CHECK_RESULT_AND_RETURN(FrameTokens.Num() > 0 && FrameTokens.Num() <= 2);
+
+						FFrameRange FrameRange;
+						FrameRange.StartFrame = FCString::Atoi(*FrameTokens[0]);
+						FrameRange.EndFrame = FrameTokens.Num() == 2 ? FCString::Atoi(*FrameTokens[1]) : FrameRange.StartFrame;
+
+						CHECK_RESULT_AND_RETURN(FrameRange.StartFrame >= 0);
+						CHECK_RESULT_AND_RETURN(FrameRange.EndFrame >= 0);
+						CHECK_RESULT_AND_RETURN(FrameRange.EndFrame >= FrameRange.StartFrame);
+
+						if (FrameRangeJson)
+						{
+							// Dropped frames are specified wrt frame range start
+							FrameRange.StartFrame -= Camera.FrameRange.Key;
+							FrameRange.EndFrame -= Camera.FrameRange.Key;
+						}
+
+						if (FrameRange.StartFrame >= 0)
+						{
+							Camera.CaptureExcludedFrames.Add(FrameRange);
+						}
+					}
+				}
+
+				FString Name = Camera.UserId;
+				CubicTakeInfo.CameraMap.Add(MoveTemp(Name), MoveTemp(Camera));
+			}
+
+			const TArray<TSharedPtr<FJsonValue>>* AudioArrayJson;
+			// Optional
+			bool HasAudio = TakeMetadataJson->TryGetArrayField(TEXT("Audio"), AudioArrayJson);
+			if (HasAudio)
+			{
+				// Audio
+				for (const TSharedPtr<FJsonValue>& AudioJson : *AudioArrayJson)
+				{
+					FStereoCameraTakeInfo::FAudio Audio;
+					const TSharedPtr<FJsonObject>& AudioObject = AudioJson->AsObject();
+
+					CHECK_RESULT_AND_RETURN(AudioObject->TryGetStringField(TEXT("UserID"), Audio.UserId));
+
+					FString StreamPath;
+					CHECK_RESULT_AND_RETURN(AudioObject->TryGetStringField(TEXT("StreamPath"), StreamPath));
+					Audio.StreamPath = ConvertPathToAbsolute(InFileName, StreamPath);
+
+					// Optional
+					AudioObject->TryGetNumberField(TEXT("TimecodeRate"), Audio.TimecodeRate);
+					AudioObject->TryGetStringField(TEXT("StartTimecode"), Audio.StartTimecode);
+
+					CubicTakeInfo.AudioArray.Add(MoveTemp(Audio));
+				}
+			}
+
+			CubicTakeInfo.TakeJsonFilePath = InFileName;
+			return MoveTemp(CubicTakeInfo);
+		}
+	}
+
+	return {};
+}
+
+TArray<FText> FStereoCameraSystemTakeParser::CheckStereoCameraTakeInfo(
+	const FString& InFilePath,
+	FStereoCameraTakeInfo& InStereoCameraTakeInfo,
+	const int32 InExpectedCameraCount,
+	const FString& InDeviceType)
+{
+	TArray<FText> Issues;
+
+	if (!FCString::IsPureAnsi(*InFilePath))
+	{
+		ReportIssue(
+			LOCTEXT("UnsupportedCharactersWithinFilePath", "Take File Path contains unsupported text characters"),
+			InStereoCameraTakeInfo,
+			Issues
+		);
+	}
+
+	if (InStereoCameraTakeInfo.DeviceInfo.Type != InDeviceType)
+	{
+		ReportIssue(
+			FText::Format(
+				LOCTEXT("UnexpectedDeviceType", "Unexpected device type: {0} instead of {1}"),
+				FText::FromString(InStereoCameraTakeInfo.DeviceInfo.Type),
+				FText::FromString(InDeviceType)
+			),
+			InStereoCameraTakeInfo,
+			Issues
+		);
+	}
+
+	if (InStereoCameraTakeInfo.CameraMap.Num() != InExpectedCameraCount)
+	{
+		ReportIssue(
+			FText::Format(
+				LOCTEXT("UnexpectedNumberOfCameras", "Unexpected number of cameras: expected {0}, found {1}"),
+				InExpectedCameraCount,
+				InStereoCameraTakeInfo.CameraMap.Num()
+			),
+			InStereoCameraTakeInfo,
+			Issues
+		);
+	}
+
+	CheckForNonANSICharacters(InStereoCameraTakeInfo, Issues);
+
+	CheckAudio(InStereoCameraTakeInfo, Issues);
+	CheckForEmptyFramesPaths(InStereoCameraTakeInfo, Issues);
+
+	return Issues;
+}
+
+
+TArray<FText> FStereoCameraSystemTakeParser::ResolveResolution(FStereoCameraTakeInfo& OutStereoCameraTakeInfo)
+{
+	using namespace UE::CaptureManager;
+
+	FResolutionResolver ResolutionResolver;
+
+	TArray<FText> Issues;
+
+	for (TPair<FString, FStereoCameraTakeInfo::FCamera>& CameraIt : OutStereoCameraTakeInfo.CameraMap)
+	{
+		const FStereoCameraTakeInfo::FCamera& Camera = CameraIt.Value;
+
+		// Doesn't take into account dropped frames:
+		const int32 CurrentCameraFrames = (Camera.FrameRange.Value - Camera.FrameRange.Key) + 1;
+
+		TValueOrError<FIntPoint, FResolutionResolver::EAddError> CameraResolution = ResolutionResolver.Add(Camera);
+
+		if (CameraResolution.HasError())
+		{
+			ReportResolutionAddError(CameraResolution.GetError(), Camera.UserId, OutStereoCameraTakeInfo, Issues);
+			continue;
+		}
+
+		CameraIt.Value.Resolution = CameraResolution.StealValue();
+	}
+
+	TValueOrError<FIntPoint, FResolutionResolver::EResolveError> Resolution = ResolutionResolver.Resolve();
+
+	if (Resolution.HasError())
+	{
+		ReportResolutionResolveError(Resolution.GetError(), OutStereoCameraTakeInfo, Issues);
+	}
+
+	return Issues;
+}
+
+TOptional<FText> CheckResolutions(const FTakeMetadata::FVideo& InCameraTakeInfo, const FCameraCalibration& InCameraCalibration)
+{
+	FIntPoint CameraCalibrationResolution(InCameraCalibration.ImageSize.X, InCameraCalibration.ImageSize.Y);
+
+	if (InCameraTakeInfo.FrameWidth.IsSet() && InCameraTakeInfo.FrameHeight.IsSet())
+	{
+		FIntPoint VideoResolution(InCameraTakeInfo.FrameWidth.GetValue(), InCameraTakeInfo.FrameHeight.GetValue());
+
+		if (VideoResolution != CameraCalibrationResolution)
+		{
+			return LOCTEXT("CheckResolutions_ResolutionValidationFailed", "Calibration and Image resolution differ");
+		}
+	}
+	else
+	{
+		UE_LOGF(LogStereoCameraMetadata, Warning, "Could not check resolutions as FrameWidth and FrameHeight have not been set for %ls at %ls", *InCameraTakeInfo.Name, *InCameraTakeInfo.Path);
+	}
+
+	return {};
+}
+
+#undef LOCTEXT_NAMESPACE
diff --git a/Engine/Plugins/VirtualProduction/CaptureManager/CaptureManagerCore/Source/CaptureMetadataExtraction/Private/StereoCameraTakeMetadata.h b/Engine/Plugins/VirtualProduction/CaptureManager/CaptureManagerCore/Source/CaptureMetadataExtraction/Private/StereoCameraTakeMetadata.h
new file mode 100644
--- /dev/null
+++ b/Engine/Plugins/VirtualProduction/CaptureManager/CaptureManagerCore/Source/CaptureMetadataExtraction/Private/StereoCameraTakeMetadata.h
@@ -0,0 +1,85 @@
+// Copyright Epic Games, Inc. All Rights Reserved.
+
+// TODO: Make Private
+
+#pragma once
+
+#include "CameraCalibration.h"
+#include "FrameRange.h"
+
+#include "Async/StopToken.h"
+
+#include "CaptureManagerTakeMetadata.h"
+
+
+class FStereoCameraInfo
+{
+public:
+
+	FString Name;
+	FCameraCalibration Calibration;
+};
+
+class FStereoCameraTakeInfo
+{
+public:
+	struct FDeviceInfo
+	{
+		FString Model;
+		FString Type;
+		FString Id;
+	};
+
+	struct FCamera
+	{
+		FString UserId;
+		TPair<uint32, uint32> FrameRange;
+		float FrameRate = 0.0f;
+		FString FramesPath;
+		FString StartTimecode;
+		FIntPoint Resolution = FIntPoint::NoneValue;
+		TArray<FFrameRange> CaptureExcludedFrames;
+	};
+	using FCameraMap = TMap<FString, FCamera>;
+
+	struct FAudio
+	{
+		FString UserId;
+		FString StreamPath;
+		float TimecodeRate = 0.0f;
+		FString StartTimecode;
+	};
+	using FAudioArray = TArray<FAudio>;
+
+	uint32 Version = 0;
+	FString Id;
+	uint32 Take = 0;
+	FString Slate;
+	FString ThumbnailPath;
+	FDateTime Date;
+	FDeviceInfo DeviceInfo;
+	FString CalibrationFilePath;
+
+	FCameraMap CameraMap;
+	FAudioArray AudioArray;
+
+	FString TakeJsonFilePath;
+
+	FString GetName() const;
+	FString GetFolderName() const;
+};
+
+class FStereoCameraSystemTakeParser
+{
+public:
+
+	static TOptional<FStereoCameraTakeInfo> ParseTakeMetadataFile(const FString& InFileName);
+
+	static TArray<FText> CheckStereoCameraTakeInfo(
+		const FString& InFilePath,
+		FStereoCameraTakeInfo& InStereoCameraTakeInfo,
+		int32 InExpectedCameraCount,
+		const FString& InDeviceType);
+
+	static TArray<FText> ResolveResolution(FStereoCameraTakeInfo& OutStereoCameraTakeInfo);
+};
diff --git a/Engine/Plugins/VirtualProduction/CaptureManager/CaptureManagerCore/Source/CaptureMetadataExtraction/Private/Utils/AppleDeviceList.cpp b/Engine/Plugins/VirtualProduction/CaptureManager/CaptureManagerCore/Source/CaptureMetadataExtraction/Private/Utils/AppleDeviceList.cpp
new file mode 100644
--- /dev/null
+++ b/Engine/Plugins/VirtualProduction/CaptureManager/CaptureManagerCore/Source/CaptureMetadataExtraction/Private/Utils/AppleDeviceList.cpp
@@ -0,0 +1,138 @@
+// Copyright Epic Games, Inc. All Rights Reserved.
+
+#include "AppleDeviceList.h"
+
+namespace UE::CaptureManager
+{
+
+const FAppleDeviceList::FDeviceMap FAppleDeviceList::DeviceMap = {
+	// iPhone
+	{ "iPhone1,1", "iPhone" },
+	{ "iPhone1,2", "iPhone 3G" },
+	{ "iPhone2,1", "iPhone 3GS" },
+	{ "iPhone3,1", "iPhone 4" },
+	{ "iPhone3,2", "iPhone 4" },
+	{ "iPhone3,3", "iPhone 4" },
+	{ "iPhone4,1", "iPhone 4S" },
+	{ "iPhone5,1", "iPhone 5" },
+	{ "iPhone5,2", "iPhone 5" },
+	{ "iPhone5,3", "iPhone 5C" },
+	{ "iPhone5,4", "iPhone 5C" },
+	{ "iPhone6,1", "iPhone 5S" },
+	{ "iPhone6,2", "iPhone 5S" },
+	{ "iPhone7,1", "iPhone 6 Plus" },
+	{ "iPhone7,2", "iPhone 6" },
+	{ "iPhone8,1", "iPhone 6S" },
+	{ "iPhone8,2", "iPhone 6S Plus" },
+	{ "iPhone8,4", "iPhone SE" },
+	{ "iPhone9,1", "iPhone 7" },
+	{ "iPhone9,2", "iPhone 7 Plus" },
+	{ "iPhone9,3", "iPhone 7"  },
+	{ "iPhone9,4", "iPhone 7 Plus" },
+	{ "iPhone10,1", "iPhone 8 " },
+	{ "iPhone10,2", "iPhone 8 Plus" },
+	{ "iPhone10,3", "iPhone X" },
+	{ "iPhone10,4", "iPhone 8" },
+	{ "iPhone10,5", "iPhone 8 Plus" },
+	{ "iPhone10,6", "iPhone X" },
+	{ "iPhone11,2", "iPhone XS" },
+	{ "iPhone11,4", "iPhone XS Max" },
+	{ "iPhone11,6", "iPhone XS Max" },
+	{ "iPhone11,8", "iPhone XR" },
+	{ "iPhone12,1", "iPhone 11" },
+	{ "iPhone12,3", "iPhone 11 Pro" },
+	{ "iPhone12,5", "iPhone 11 Pro Max" },
+	{ "iPhone12,8", "iPhone SE 2" },
+	{ "iPhone13,1", "iPhone 12 Mini" },
+	{ "iPhone13,2", "iPhone 12" },
+	{ "iPhone13,3", "iPhone 12 Pro" },
+	{ "iPhone13,4", "iPhone 12 Pro Max" },
+	{ "iPhone14,2", "iPhone 13 Pro" },
+	{ "iPhone14,3", "iPhone 13 Pro Max" },
+	{ "iPhone14,4", "iPhone 13 Mini" },
+	{ "iPhone14,5", "iPhone 13" },
+	{ "iPhone14,6", "iPhone SE 3" },
+	{ "iPhone14,7", "iPhone 14" },
+	{ "iPhone14,8", "iPhone 14 Plus" },
+	{ "iPhone15,2", "iPhone 14 Pro" },
+	{ "iPhone15,3", "iPhone 14 Pro Max" },
+
+	// iPad
+	{ "iPad1,1", "iPad" },
+	{ "iPad2,1", "iPad 2" },
+	{ "iPad2,2", "iPad 2" },
+	{ "iPad2,3", "iPad 2" },
+	{ "iPad2,4", "iPad 2" },
+	{ "iPad3,1", "iPad 3" },
+	{ "iPad3,2", "iPad 3" },
+	{ "iPad3,3", "iPad 3" },
+	{ "iPad3,4", "iPad 4" },
+	{ "iPad3,5", "iPad 4" },
+	{ "iPad3,6", "iPad 4" },
+	{ "iPad6,11", "iPad 5" },
+	{ "iPad6,12", "iPad 5" },
+	{ "iPad7,5", "iPad 6" },
+	{ "iPad7,6", "iPad 6" },
+	{ "iPad7,11", "iPad 7" },
+	{ "iPad7,12", "iPad 7" },
+	{ "iPad11,6", "iPad 8" },
+	{ "iPad11,7", "iPad 8" },
+	{ "iPad12,1", "iPad 9" },
+	{ "iPad12,2", "iPad 9" },
+	{ "iPad4,1", "iPad Air" },
+	{ "iPad4,2", "iPad Air" },
+	{ "iPad4,3", "iPad Air" },
+	{ "iPad5,3", "iPad Air 2" },
+	{ "iPad5,4", "iPad Air 2" },
+	{ "iPad11,3", "iPad Air 3" },
+	{ "iPad11,4", "iPad Air 3" },
+	{ "iPad13,1", "iPad Air 4" },
+	{ "iPad13,2", "iPad Air 4" },
+	{ "iPad13,16", "iPad Air 5" },
+	{ "iPad13,17", "iPad Air 5" },
+	{ "iPad2,5", "iPad Mini" },
+	{ "iPad2,6", "iPad Mini" },
+	{ "iPad2,7", "iPad Mini" },
+	{ "iPad4,4", "iPad Mini 2" },
+	{ "iPad4,5", "iPad Mini 2" },
+	{ "iPad4,6", "iPad Mini 2" },
+	{ "iPad4,7", "iPad Mini 3" },
+	{ "iPad4,8", "iPad Mini 3" },
+	{ "iPad4,9", "iPad Mini 3" },
+	{ "iPad5,1", "iPad Mini 4" },
+	{ "iPad5,2", "iPad Mini 4" },
+	{ "iPad11,1", "iPad Mini 5" },
+	{ "iPad11,2", "iPad Mini 5" },
+	{ "iPad14,1", "iPad Mini 6" },
+	{ "iPad14,2", "iPad Mini 6" },
+	{ "iPad6,3", "iPad Pro 9.7-inch" },
+	{ "iPad6,4", "iPad Pro 9.7-inch" },
+	{ "iPad7,3", "iPad Pro 10.5-inch" },
+	{ "iPad7,4", "iPad Pro 10.5-inch" },
+	{ "iPad8,1", "iPad Pro 11-inch" },
+	{ "iPad8,2", "iPad Pro 11-inch" },
+	{ "iPad8,3", "iPad Pro 11-inch" },
+	{ "iPad8,4", "iPad Pro 11-inch" },
+	{ "iPad8,9", "iPad Pro 11-inch 2" },
+	{ "iPad8,10", "iPad Pro 11-inch 2" },
+	{ "iPad13,4", "iPad Pro 11-inch 3" },
+	{ "iPad13,5", "iPad Pro 11-inch 3" },
+	{ "iPad13,6", "iPad Pro 11-inch 3" },
+	{ "iPad13,7", "iPad Pro 11-inch 3" },
+	{ "iPad6,7", "iPad Pro 12.9-inch" },
+	{ "iPad6,8", "iPad Pro 12.9-inch" },
+	{ "iPad7,1", "iPad Pro 12.9-inch 2" },
+	{ "iPad7,2", "iPad Pro 12.9-inch 2" },
+	{ "iPad8,5", "iPad Pro 12.9-inch 3" },
+	{ "iPad8,6", "iPad Pro 12.9-inch 3" },
+	{ "iPad8,7", "iPad Pro 12.9-inch 3" },
+	{ "iPad8,8", "iPad Pro 12.9-inch 3" },
+	{ "iPad8,11", "iPad Pro 12.9-inch 4" },
+	{ "iPad8,12", "iPad Pro 12.9-inch 4" },
+	{ "iPad13,8", "iPad Pro 12.9-inch 5" },
+	{ "iPad13,9", "iPad Pro 12.9-inch 5" },
+	{ "iPad13,10", "iPad Pro 12.9-inch 5" },
+	{ "iPad13,11", "iPad Pro 12.9-inch 5" }
+};
+
+}
\ No newline at end of file
diff --git a/Engine/Plugins/VirtualProduction/CaptureManager/CaptureManagerCore/Source/CaptureMetadataExtraction/Private/Utils/AppleDeviceList.h b/Engine/Plugins/VirtualProduction/CaptureManager/CaptureManagerCore/Source/CaptureMetadataExtraction/Private/Utils/AppleDeviceList.h
new file mode 100644
--- /dev/null
+++ b/Engine/Plugins/VirtualProduction/CaptureManager/CaptureManagerCore/Source/CaptureMetadataExtraction/Private/Utils/AppleDeviceList.h
@@ -0,0 +1,19 @@
+// Copyright Epic Games, Inc. All Rights Reserved.
+
+#pragma once
+
+#include "Containers/Map.h"
+#include "Containers/UnrealString.h"
+
+namespace UE::CaptureManager
+{
+
+class FAppleDeviceList
+{
+public:
+	using FDeviceMap = TMap<FString, FString>;
+
+	static const FDeviceMap DeviceMap;
+};
+
+}
\ No newline at end of file
diff --git a/Engine/Plugins/VirtualProduction/CaptureManager/CaptureManagerCore/Source/CaptureMetadataExtraction/Private/Utils/CaptureExtractTimecode.cpp b/Engine/Plugins/VirtualProduction/CaptureManager/CaptureManagerCore/Source/CaptureMetadataExtraction/Private/Utils/CaptureExtractTimecode.cpp
new file mode 100644
--- /dev/null
+++ b/Engine/Plugins/VirtualProduction/CaptureManager/CaptureManagerCore/Source/CaptureMetadataExtraction/Private/Utils/CaptureExtractTimecode.cpp
@@ -0,0 +1,611 @@
+// Copyright Epic Games, Inc. All Rights Reserved.
+
+#include "Utils/CaptureExtractTimecode.h"
+
+#include "Audio.h"
+#include "Async/Async.h"
+#include "CaptureManagerMediaRWModule.h"
+#include "Misc/FileHelper.h"
+#include "MediaPlaylist.h"
+#include "Sound/SoundWaveTimecodeInfo.h"
+#include "UObject/Package.h"
+//#include "Settings/CaptureManagerSettings.h"
+#include "Logging/LogMacros.h"
+#include "IElectraPlayerPluginModule.h"
+#include "IMediaPlayer.h"
+#include "IMediaOptions.h"
+#include "IMediaEventSink.h"
+#include "IMediaTracks.h"
+#include "Serialization/JsonSerializer.h"
+
+#include "ParseTakeUtils.h"
+#include "SmpteTimecodeUtils.h"
+
+#include "ProcessRunner/ProcessRunner.h"
+
+#define LOCTEXT_NAMESPACE "CaptureExtractTimecode"
+
+DEFINE_LOG_CATEGORY_STATIC(LogCaptureExtractTimecode, Log, All);
+
+
+
+static bool IsEncoderFFmpeg(const FString& InEncoderPath)
+{
+	// Get base filename without extension
+	const FString EncoderFileName = FPaths::GetBaseFilename(InEncoderPath);
+	return EncoderFileName.Equals(TEXT("ffmpeg"), ESearchCase::IgnoreCase);
+}
+
+static FString GetFFprobePath(const FString& InEncoderPath)
+{
+	check(IsEncoderFFmpeg(InEncoderPath));
+
+	// Get base filename with extension: e.g ffmpeg.exe
+	const FString EncoderFileName = FPaths::GetCleanFilename(InEncoderPath);
+	const FString RootPath = FPaths::GetPath(InEncoderPath);
+	const FString FFprobeFileName = EncoderFileName.Replace(TEXT("ffmpeg"), TEXT("ffprobe"), ESearchCase::IgnoreCase);
+	const FString FFprobePath = RootPath / FFprobeFileName;
+	return FFprobePath;
+}
+
+FCaptureExtractVideoInfo::FResult FCaptureExtractVideoInfo::Create(const FString& InFilePath, TOptional<UE::CaptureManager::Private::FFProbeCommand> InProbeCommand)
+{
+	using namespace UE::CaptureManager;
+
+	FCaptureExtractVideoInfo VideoInfoExtractor;
+
+	FExtractResult Result = VideoInfoExtractor.ExtractInfo(InFilePath, InProbeCommand);
+
+	if (Result.HasError())
+	{
+		return MakeError(Result.StealError());
+	}
+
+	UE_LOGF(LogCaptureExtractTimecode, Verbose, "Extracted timecode '%ls' from video file: %ls", *VideoInfoExtractor.GetTimecode().ToString(), *InFilePath);
+
+	return MakeValue(MoveTemp(VideoInfoExtractor));
+}
+
+FFrameRate FCaptureExtractVideoInfo::GetFrameRate() const
+{
+	return FrameRate;
+}
+
+FTimecode FCaptureExtractVideoInfo::GetTimecode() const
+{
+	return TimecodeInfo.Timecode;
+}
+
+FFrameRate FCaptureExtractVideoInfo::GetTimecodeRate() const
+{
+	return TimecodeInfo.TimecodeRate;
+}
+
+bool FCaptureExtractVideoInfo::ContainsAudio() const
+{
+	return bContainsAudio;
+}
+
+float FCaptureExtractVideoInfo::GetAudioDurationSeconds() const
+{
+	return AudioDurationSeconds;
+}
+
+EMediaOrientation FCaptureExtractVideoInfo::GetVideoOrientation() const
+{
+	return VideoOrientation;
+}
+
+FCaptureExtractVideoInfo::FCaptureExtractVideoInfo()
+	: FrameRate(1, 1)
+{
+}
+
+FCaptureExtractVideoInfo::FExtractResult FCaptureExtractVideoInfo::ExtractInfo(const FString& InFilePath, TOptional<UE::CaptureManager::Private::FFProbeCommand> InProbeCommand)
+{
+	using namespace UE::CaptureManager;
+	using namespace UE::CaptureManager::Private;
+
+	// Replace any backslashes to avoid escape code issues while logging the file path
+	const FString FilePath = InFilePath.Replace(TEXT("\\"), TEXT("/"));
+	if (FilePath.IsEmpty())
+	{
+		return MakeError(ECaptureExtractInfoError::InternalError);
+	}
+
+	if (InProbeCommand)
+	{
+		return ExtractInfoUsingFFProbe(InFilePath, *InProbeCommand);
+	}
+	else
+	{
+		// We gather audio duration via an audio reader if available and the format is supported.
+		FMediaRWManager& MediaManager = FModuleManager::LoadModuleChecked<FCaptureManagerMediaRWModule>("CaptureManagerMediaRW").Get();
+		TValueOrError<TUniquePtr<IAudioReader>, FText> AudioReaderResult = MediaManager.CreateAudioReader(InFilePath);
+		if (AudioReaderResult.HasValue())
+		{
+			bContainsAudio = true;
+			AudioDurationSeconds = static_cast<float>(AudioReaderResult.StealValue()->GetDuration().GetTotalSeconds());
+		}
+
+		// Frame rate, and timecode data is read from the electra player
+		return ExtractInfoUsingElectraPlayer(InFilePath);
+	}
+}
+
+FCaptureExtractVideoInfo::FExtractResult FCaptureExtractVideoInfo::ExtractInfoUsingElectraPlayer(const FString& InFilePath)
+{
+	using namespace UE::CaptureManager;
+
+	using FEventSinkResult = FExtractResult;
+	TPromise<FExtractResult> EventSinkPromise;
+
+	class FMediaEventSink : public IMediaEventSink
+	{
+	public:
+
+		FMediaEventSink(TPromise<FEventSinkResult>& InPromise)
+			: Promise(InPromise)
+			, bReleased(false)
+		{
+		}
+
+		virtual void ReceiveMediaEvent(EMediaEvent Event) override
+		{
+			if (bReleased)
+			{
+				return;
+			}
+
+			if (Event == EMediaEvent::MediaOpened)
+			{
+				Promise.SetValue(MakeValue());
+			}
+			else if (Event == EMediaEvent::MediaOpenFailed)
+			{
+				Promise.SetValue(MakeError(ECaptureExtractInfoError::UnableToOpenMedia));
+			}
+		}
+
+		void Release()
+		{
+			bReleased = true;
+		}
+
+	private:
+
+		TPromise<FEventSinkResult>& Promise;
+		std::atomic_bool bReleased;
+
+	} MediaEventSink(EventSinkPromise);
+
+	IElectraPlayerPluginModule& ElectraModule = FModuleManager::LoadModuleChecked<IElectraPlayerPluginModule>("ElectraPlayerPlugin");
+
+	TSharedPtr<IMediaPlayer, ESPMode::ThreadSafe> MediaPlayer = ElectraModule.CreatePlayer(MediaEventSink);
+	if (!MediaPlayer)
+	{
+		// Broken promises are considered errors
+		FExtractResult Error = MakeError(ECaptureExtractInfoError::InternalError);
+		EventSinkPromise.SetValue(Error);
+		return Error;
+	}
+
+	UDesiredPlayerMediaSource* MediaSource = NewObject<UDesiredPlayerMediaSource>();
+	MediaSource->SetFilePath(InFilePath);
+
+	FString FileUrl = TEXT("file://") + InFilePath;
+
+	FMediaPlayerOptions PlayerOptions;
+	PlayerOptions.SetAllAsOptional();
+	PlayerOptions.InternalCustomOptions.Emplace(MediaPlayerOptionValues::ParseTimecodeInfo(), FVariant());
+
+	MediaPlayer->Open(FileUrl, MediaSource, &PlayerOptions);
+
+	FString TimecodeString;
+	FString TimecodeRateString;
+
+	FDateTime EventWaitStart = FDateTime::Now();
+
+	TFuture<FEventSinkResult> EventSinkFuture = EventSinkPromise.GetFuture();
+
+	bool bIsTimeout = true;
+
+	while ((FDateTime::Now() - EventWaitStart).GetSeconds() < TimeoutPeriod)
+	{
+		MediaPlayer->TickInput(0, 0);
+
+		if (EventSinkFuture.WaitFor(FTimespan::FromMilliseconds(100.0f))) // Wait 100 milliseconds before invoking TickInput again
+		{
+			FEventSinkResult Result = EventSinkFuture.Get();
+			if (Result.HasError())
+			{
+				return Result;
+			}
+
+			bIsTimeout = false;
+
+			const FVariant Timecode = MediaPlayer->GetMediaInfo(UMediaPlayer::MediaInfoNameStartTimecodeValue.Resolve());
+			if (!Timecode.IsEmpty())
+			{
+				TimecodeString = Timecode.GetValue<FString>();
+
+				const FVariant TimecodeRate = MediaPlayer->GetMediaInfo(UMediaPlayer::MediaInfoNameStartTimecodeFrameRate.Resolve());
+				if (!TimecodeRate.IsEmpty())
+				{
+					TimecodeRateString = TimecodeRate.GetValue<FString>();
+				}
+			}
+
+			int32 Track = MediaPlayer->GetTracks().GetSelectedTrack(EMediaTrackType::Video);
+			int32 Format = MediaPlayer->GetTracks().GetTrackFormat(EMediaTrackType::Video, Track);
+
+			if (Track != INDEX_NONE && Format != INDEX_NONE)
+			{
+				FMediaVideoTrackFormat FormatInfo;
+				MediaPlayer->GetTracks().GetVideoTrackFormat(Track, Format, FormatInfo);
+
+				FrameRate = ConvertFrameRate(FormatInfo.FrameRate);
+			}
+			else
+			{
+				UE_LOGF(LogCaptureExtractTimecode, Warning, "Failed to obtain the frame rate");
+			}
+
+			break;
+		}
+	}
+
+	// Making sure that no additional events will arrive
+	MediaEventSink.Release();
+	MediaPlayer->Close();
+
+	if (bIsTimeout)
+	{
+		FExtractResult Error = MakeError(ECaptureExtractInfoError::InternalError);
+
+		if (!EventSinkFuture.IsReady())
+		{
+			// Broken promises are considered errors
+			EventSinkPromise.SetValue(Error);
+		}
+
+		return Error;
+	}
+
+	if (!TimecodeString.IsEmpty())
+	{
+		const TOptional<FTimecode> Timecode = FTimecode::ParseTimecode(*TimecodeString);
+		if (!Timecode.IsSet())
+		{
+			UE_LOGF(LogCaptureExtractTimecode, Warning, "Failed to parse the timecode");
+		}
+		else
+		{
+			TimecodeInfo.Timecode = Timecode.GetValue();
+		}
+	}
+	else
+	{
+		UE_LOGF(LogCaptureExtractTimecode, Warning, "Timecode has not been found");
+	}
+
+	if (!TimecodeRateString.IsEmpty())
+	{
+		TValueOrError<FFrameRate, UE::CaptureManager::ECaptureExtractInfoError> Result = ParseTimecodeRate(TimecodeRateString);
+		if (Result.IsValid())
+		{
+			TimecodeInfo.TimecodeRate = Result.GetValue();
+		}
+	}
+	else
+	{
+		UE_LOGF(LogCaptureExtractTimecode, Warning, "Timecode rate has not been found. Using video Frame rate as timecode rate");
+		TimecodeInfo.TimecodeRate = FrameRate;
+	}
+
+	return MakeValue();
+}
+
+FCaptureExtractVideoInfo::FExtractResult FCaptureExtractVideoInfo::ExtractInfoUsingFFProbe(const FString& InFilePath, const UE::CaptureManager::Private::FFProbeCommand& InProbeCommand)
+{
+	using namespace UE::CaptureManager;
+
+	FProcessRunnerResult RunnerResult = InProbeCommand.Execute(InFilePath, TimeoutPeriod);
+
+	if (!RunnerResult.HasValue())
+	{
+		return MakeError(ECaptureExtractInfoError::InternalError);
+	}
+
+	TArray<uint8> CommandOutput = RunnerResult.StealValue();
+	FString JsonString = FString::ConstructFromPtrSize(reinterpret_cast<const UTF8CHAR*>(CommandOutput.GetData()), CommandOutput.Num());
+
+	TSharedPtr<FJsonObject> JsonObject;
+	if (!FJsonSerializer::Deserialize(TJsonReaderFactory<>::Create(JsonString), JsonObject))
+	{
+		UE_LOGF(LogCaptureExtractTimecode, Error, "Failed to parse ffprobe output for file \"%ls\". Ensure ffprobe is up to date.", *FPaths::GetCleanFilename(InFilePath));
+		return MakeError(ECaptureExtractInfoError::FailedToParseJson);
+	}
+
+	// Streams
+	const TArray<TSharedPtr<FJsonValue>>* StreamsArray;
+	if (JsonObject->TryGetArrayField(TEXT("streams"), StreamsArray))
+	{
+		for (const TSharedPtr<FJsonValue>& Value : *StreamsArray)
+		{
+			TSharedPtr<FJsonObject> StreamObject = Value->AsObject();
+			if (StreamObject.IsValid())
+			{
+				FString CodecType;
+				if (!StreamObject->TryGetStringField(TEXT("codec_type"), CodecType)) 
+				{
+					continue;
+				}
+
+				if (CodecType == TEXT("audio"))
+				{
+					// Audio Duration
+					float Duration = 0.0;
+					if (StreamObject->TryGetNumberField(TEXT("duration"), Duration))
+					{
+						bContainsAudio = true;
+						AudioDurationSeconds = Duration;
+					}
+				}
+				else if (CodecType == TEXT("video"))
+				{
+					// Frame Rate
+					FString FrameRateString;
+					if (StreamObject->TryGetStringField(TEXT("r_frame_rate"), FrameRateString)) 
+					{
+						TArray<FString> FrameRateComponents;
+						FrameRateString.ParseIntoArray(FrameRateComponents, TEXT("/"));
+						if (FrameRateComponents.Num() == 2)
+						{
+							int32 Numerator = FCString::Atoi(*FrameRateComponents[0]);
+							int32 Denominator = FCString::Atoi(*FrameRateComponents[1]);
+							FrameRate = FFrameRate(Numerator, Denominator);
+							TimecodeInfo.TimecodeRate = FrameRate;
+						}
+					}
+
+					// Video Orientation
+					int32 Rotation = 0;
+					const TArray<TSharedPtr<FJsonValue>>* SideDataArrayPointer;
+					if (StreamObject->TryGetArrayField(TEXT("side_data_list"), SideDataArrayPointer))
+					{
+						const TArray<TSharedPtr<FJsonValue>> SideDataArray = *SideDataArrayPointer;
+						if (SideDataArray.Num() > 0)
+						{
+							const TSharedPtr<FJsonObject>* SideDataObjectPointer;
+							if (SideDataArray[0]->TryGetObject(SideDataObjectPointer))
+							{
+								(*SideDataObjectPointer)->TryGetNumberField(TEXT("rotation"), Rotation);
+							}
+						}
+					}
+					switch (FMath::WrapExclusive(Rotation, 0, 360))
+					{
+						case 90:
+							VideoOrientation = EMediaOrientation::CW90;
+							break;
+						case 180:
+							VideoOrientation = EMediaOrientation::CW180;
+							break;
+						case 270:
+							VideoOrientation = EMediaOrientation::CW270;
+							break;
+						default:
+							VideoOrientation = EMediaOrientation::Original;
+					}
+				}
+			}
+		}
+	}
+
+	// Timecode
+	const TSharedPtr<FJsonObject>* FormatObject;
+	if (JsonObject->TryGetObjectField(TEXT("format"), FormatObject)) 
+	{
+		const TSharedPtr<FJsonObject>* TagsObject;
+		if ((*FormatObject)->TryGetObjectField(TEXT("tags"), TagsObject)) 
+		{
+			FString TimecodeString;
+			if ((*TagsObject)->TryGetStringField(TEXT("timecode"), TimecodeString))
+			{
+				TArray<FString> TimecodeComponents;
+				TimecodeString.ParseIntoArray(TimecodeComponents, TEXT(":")); // We don't support DF timecode
+				if (TimecodeComponents.Num() == 4)
+				{
+					int32 Hours = FCString::Atoi(*TimecodeComponents[0]);
+					int32 Minutes = FCString::Atoi(*TimecodeComponents[1]);
+					int32 Seconds = FCString::Atoi(*TimecodeComponents[2]);
+					int32 Frames = FCString::Atoi(*TimecodeComponents[3]);
+
+					constexpr bool bIsDropFrame = false;
+					TimecodeInfo.Timecode = FTimecode(Hours, Minutes, Seconds, Frames, bIsDropFrame);
+				}
+			}
+		}
+	}
+
+	return MakeValue();
+}
+
+TValueOrError<FFrameRate, UE::CaptureManager::ECaptureExtractInfoError> FCaptureExtractVideoInfo::ParseTimecodeRate(FString TimecodeRateString)
+{
+	using namespace UE::CaptureManager;
+
+	if (TimecodeRateString.IsEmpty())
+	{
+		return MakeError(ECaptureExtractInfoError::TimecodeRateNotFound);
+	}
+	else
+	{
+		// TimecodeRateString is made out of FFrameRate::ToPrettyText()
+		// It could be "{0} fps" or "{0} s"
+
+		FFrameRate TimecodeRate;
+
+		FString Left;
+		FString Right;
+		TimecodeRateString.Split(" ", &Left, &Right);
+
+		double Number = FCString::Atoi(*Left);
+		uint32 IntNumber;
+		uint32 Multiplier = 1;
+
+
+		for (; Multiplier <= 10000; Multiplier *= 10)
+		{
+			double TmpNumber = Number * Multiplier;
+			IntNumber = FMath::RoundToInt32(TmpNumber);
+			if (FMath::Abs(TmpNumber - IntNumber) < 0.01)
+			{
+				break;
+			}
+		}
+
+		uint32 Nominator = 0;
+		uint32 Denominator = 0;
+
+		if (Right == "fps")
+		{
+			Nominator = IntNumber;
+			Denominator = Multiplier;
+		}
+		else if (Right == "s")
+		{
+			Nominator = Multiplier;
+			Denominator = IntNumber;
+		}
+
+		TimecodeRate = FFrameRate(Nominator, Denominator);
+
+		return MakeValue(TimecodeRate);
+	}
+}
+
+FCaptureExtractAudioTimecode::FCaptureExtractAudioTimecode(const FString& InFilePath)
+	: FilePath(InFilePath)
+{}
+
+FCaptureExtractAudioTimecode::FTimecodeInfoResult FCaptureExtractAudioTimecode::Extract()
+{
+	return ExtractTimecodeFromBroadcastWaveFormat(FFrameRate());
+}
+
+FCaptureExtractAudioTimecode::FTimecodeInfoResult FCaptureExtractAudioTimecode::Extract(FFrameRate InFrameRate)
+{
+	using namespace UE::CaptureManager;
+
+	check(!FilePath.IsEmpty());
+
+	FTimecodeInfoResult Result = MakeError(ECaptureExtractInfoError::UnhandledMedia);
+
+	FString FileExtension = FPaths::GetExtension(FilePath);
+	if (FileExtension == "wav")
+	{
+		// Convert timecode rate to SMPTE timecode rate
+		FFrameRate TimecodeRate = UE::CaptureData::EstimateSmpteTimecodeRate(InFrameRate);
+		Result = ExtractTimecodeFromBroadcastWaveFormat(TimecodeRate);
+	}
+
+	if (Result.HasValue())
+	{
+		UE_LOGF(LogCaptureExtractTimecode, Verbose, "Extracted timecode '%ls' from audio file: %ls", *Result.GetValue().Timecode.ToString(), *FilePath);
+	}
+	else
+	{
+		UE_LOGF(LogCaptureExtractTimecode, Log, "Timecode not found for audio file: %ls", *FilePath);
+	}
+
+	return Result;
+}
+
+FCaptureExtractAudioTimecode::FTimecodeInfoResult FCaptureExtractAudioTimecode::ExtractTimecodeFromBroadcastWaveFormat(FFrameRate InTimecodeRate)
+{
+	using namespace UE::CaptureManager;
+
+	TArray<uint8> WAVData;
+	if (FFileHelper::LoadFileToArray(WAVData, *FilePath))
+	{
+		FWaveModInfo WAVInfo;
+		if (WAVInfo.ReadWaveInfo(WAVData.GetData(), WAVData.Num()))
+		{
+			if (WAVInfo.TimecodeInfo.IsValid())
+			{
+				FSoundWaveTimecodeInfo TimecodeInfo = *WAVInfo.TimecodeInfo.Get();
+				const double NumSecondsSinceMidnight = TimecodeInfo.GetNumSecondsSinceMidnight();
+
+				FFrameRate TimecodeRate = TimecodeInfo.TimecodeRate;
+				const bool bTimecodeRateIsSampleRate = TimecodeRate == FFrameRate(TimecodeInfo.NumSamplesPerSecond, 1);
+				if (bTimecodeRateIsSampleRate)
+				{
+					UE_LOGF(
+						LogCaptureExtractTimecode,
+						Log,
+						"Embedded timecode rate is %.2f fps (the sample rate). "
+						"This usually indicates there is no timecode rate information in the wav file: %ls",
+						TimecodeInfo.TimecodeRate.AsDecimal(),
+						*FilePath
+					);
+
+					if (InTimecodeRate != FFrameRate())
+					{
+						// Use the provided timecode rate instead
+						TimecodeRate = InTimecodeRate;
+
+						UE_LOGF(
+							LogCaptureExtractTimecode,
+							Log,
+							"Taking the embedded audio timecode but estimating an SMPTE audio timecode rate. "
+							"Timecode rate for %ls set to %.2f",
+							*FilePath,
+							TimecodeRate.AsDecimal()
+						);
+					}
+				}
+
+				FTimecode AudioTimecode = FTimecode(NumSecondsSinceMidnight, TimecodeRate, TimecodeInfo.bTimecodeIsDropFrame, /* InbRollover = */ true);
+				FTimecodeInfo TimecodeAndRate
+				{ 
+					.Timecode = MoveTemp(AudioTimecode), 
+					.TimecodeRate = MoveTemp(TimecodeRate)
+				};
+
+				return MakeValue(MoveTemp(TimecodeAndRate));
+			}
+		}
+	}
+
+	return MakeError(ECaptureExtractInfoError::TimecodeNotFound);
+}
+
+namespace UE::CaptureManager::Private
+{
+
+FFProbeCommand::FFProbeCommand(FString InFfmpegPath) :
+	ProbePath(GetFFprobePath(InFfmpegPath))
+{
+}
+
+FProcessRunnerResult FFProbeCommand::Execute(const FString& InFilePath, int32 InTimeoutSeconds, const FStopToken* InStopToken) const
+{
+	static constexpr TCHAR ArgsFmt[] = TEXT("-v error -show_entries stream=codec_type,r_frame_rate,duration:format_tags=timecode:stream_side_data=rotation -of json \"{0}\"");
+
+	if (!FPaths::FileExists(ProbePath))
+	{
+		UE_LOGF(LogCaptureExtractTimecode, Warning, "Could not extract timecode for '%ls' using ffprobe. ffprobe could not be found at path '%ls'. ffprobe is expected to be found in the same directory as ffmpeg.", *InFilePath, *ProbePath);
+		return MakeError(EProcessRunnerError::LaunchFailed);
+	}
+
+	const FString Args = FString::Format(ArgsFmt, { InFilePath });
+	return FProcessRunner::Run(ProbePath, Args, InStopToken, InTimeoutSeconds);
+}
+
+}
+
+#undef LOCTEXT_NAMESPACE
+
+
diff --git a/Engine/Plugins/VirtualProduction/CaptureManager/CaptureManagerCore/Source/CaptureMetadataExtraction/Private/Utils/CaptureExtractTimecode.h b/Engine/Plugins/VirtualProduction/CaptureManager/CaptureManagerCore/Source/CaptureMetadataExtraction/Private/Utils/CaptureExtractTimecode.h
new file mode 100644
--- /dev/null
+++ b/Engine/Plugins/VirtualProduction/CaptureManager/CaptureManagerCore/Source/CaptureMetadataExtraction/Private/Utils/CaptureExtractTimecode.h
@@ -0,0 +1,134 @@
+// Copyright Epic Games, Inc. All Rights Reserved.
+
+#pragma once
+
+#include "FileMediaSource.h"
+
+#include "MediaSample.h"
+#include "MediaPlayer.h"
+
+#include "ProcessRunner/ProcessRunner.h"
+#include "Templates/ValueOrError.h"
+
+#include "IMediaClock.h"
+#include "IMediaClockSink.h"
+#include "IMediaModule.h"
+
+#include "CaptureExtractTimecode.generated.h"
+
+#define UE_API CAPTUREMETADATAEXTRACTION_API
+
+namespace UE::CaptureManager
+{
+
+struct FTimecodeInfo
+{
+	FTimecode Timecode;
+	FFrameRate TimecodeRate;
+};
+
+struct FVideoInformation
+{
+	FFrameRate FrameRate;
+	FTimecodeInfo TimecodeInfo;
+};
+
+enum ECaptureExtractInfoError : int32
+{
+	InternalError = -1,
+	TimecodeNotFound = 1,
+	UnableToParseTimecode = 2,
+	UnableToParseTimecodeRate = 3,
+	TimecodeRateNotFound = 4,
+	UnhandledMedia = 5,
+	UnableToOpenMedia = 6,
+	FfprobeNotFound = 7,
+	FailedToParseJson = 8,
+};
+}
+
+#if WITH_EDITOR
+UCLASS(BlueprintType, MinimalAPI)
+class UDesiredPlayerMediaSource : public UFileMediaSource
+{
+	GENERATED_BODY()
+public:
+	//~ IMediaOptions interface
+	virtual FName GetDesiredPlayerName() const override
+	{
+		return TEXT("ElectraPlayer");
+	}
+};
+#endif // WITH_EDITOR
+
+namespace UE::CaptureManager::Private
+{
+class UE_API FFProbeCommand
+{
+public:
+	FFProbeCommand(FString InFfmpegPath);
+
+	FProcessRunnerResult Execute(const FString& InFilePath, int32 InTimeoutSeconds, const FStopToken* InStopToken = nullptr) const;
+
+private:
+	const FString ProbePath;
+};
+}
+
+class FCaptureExtractVideoInfo final
+{
+public:
+
+	using FResult = TValueOrError<FCaptureExtractVideoInfo, UE::CaptureManager::ECaptureExtractInfoError>;
+
+	UE_API static FResult Create(const FString& InFilePath, TOptional<UE::CaptureManager::Private::FFProbeCommand> InProbeCommand /*= {}*/);
+
+	UE_API FFrameRate GetFrameRate() const;
+	UE_API FTimecode GetTimecode() const;
+	UE_API FFrameRate GetTimecodeRate() const;
+	UE_API bool ContainsAudio() const;
+	UE_API float GetAudioDurationSeconds() const;
+	UE_API EMediaOrientation GetVideoOrientation() const;
+
+private:
+
+	FCaptureExtractVideoInfo();
+
+	using FExtractResult = TValueOrError<void, UE::CaptureManager::ECaptureExtractInfoError>;
+	FExtractResult ExtractInfo(const FString& InFilePath, TOptional<UE::CaptureManager::Private::FFProbeCommand> InProbeCommand);
+	FExtractResult ExtractInfoUsingElectraPlayer(const FString& InFilePath);
+	FExtractResult ExtractInfoUsingFFProbe(const FString& InFilePath, const UE::CaptureManager::Private::FFProbeCommand& InProbeCommand);
+
+	TValueOrError<FFrameRate, UE::CaptureManager::ECaptureExtractInfoError> ParseTimecodeRate(FString TimecodeRateString);
+
+	static constexpr int32 TimeoutPeriod = 3;
+
+	FFrameRate FrameRate;
+	UE::CaptureManager::FTimecodeInfo TimecodeInfo;
+	bool bContainsAudio = false;
+	float AudioDurationSeconds = 0.0f;
+	EMediaOrientation VideoOrientation = EMediaOrientation::Original;
+};
+
+class FCaptureExtractAudioTimecode final : public TSharedFromThis<FCaptureExtractAudioTimecode>
+{
+public:
+
+	using FTimecodeInfoResult = TValueOrError<UE::CaptureManager::FTimecodeInfo, UE::CaptureManager::ECaptureExtractInfoError>;
+
+	UE_API FCaptureExtractAudioTimecode(const FString& InFilePath);
+	UE_API ~FCaptureExtractAudioTimecode() = default;
+
+	UE_API FTimecodeInfoResult Extract();
+	UE_API FTimecodeInfoResult Extract(FFrameRate InFrameRate);
+
+private:
+
+	FString FilePath;
+
+	const int32 TimeoutPeriod = 3;
+
+	FTimecodeInfoResult ExtractTimecodeFromBroadcastWaveFormat(FFrameRate InTimecodeRate);
+};
+
+#undef UE_API
diff --git a/Engine/Plugins/VirtualProduction/CaptureManager/CaptureManagerCore/Source/DataIngestCore/Private/IngestCaptureData.cpp b/Engine/Plugins/VirtualProduction/CaptureManager/CaptureManagerCore/Source/DataIngestCore/Private/IngestCaptureData.cpp
--- a/Engine/Plugins/VirtualProduction/CaptureManager/CaptureManagerCore/Source/DataIngestCore/Private/IngestCaptureData.cpp
+++ b/Engine/Plugins/VirtualProduction/CaptureManager/CaptureManagerCore/Source/DataIngestCore/Private/IngestCaptureData.cpp
@@ -381,4 +381,27 @@ TOptional<FText> Serialize(const FString& InFilePath, const FString& InFileName,
 
 }
 
+void FIngestCaptureData::MakePathsAbsolute(const FString& InBasePath)
+{
+	for (FVideo& VideoEntry : Video)
+	{
+		VideoEntry.Path = FPaths::ConvertRelativePathToFull(InBasePath, VideoEntry.Path);
+	}
+
+	for (FVideo& DepthEntry : Depth)
+	{
+		DepthEntry.Path = FPaths::ConvertRelativePathToFull(InBasePath, DepthEntry.Path);
+	}
+
+	for (FAudio& AudioEntry : Audio)
+	{
+		AudioEntry.Path = FPaths::ConvertRelativePathToFull(InBasePath, AudioEntry.Path);
+	}
+
+	for (FCalibration& CalibrationEntry : Calibration)
+	{
+		CalibrationEntry.Path = FPaths::ConvertRelativePathToFull(InBasePath, CalibrationEntry.Path);
+	}
+}
+
 #undef LOCTEXT_NAMESPACE
\ No newline at end of file
diff --git a/Engine/Plugins/VirtualProduction/CaptureManager/CaptureManagerCore/Source/DataIngestCore/Public/IngestCaptureData.h b/Engine/Plugins/VirtualProduction/CaptureManager/CaptureManagerCore/Source/DataIngestCore/Public/IngestCaptureData.h
--- a/Engine/Plugins/VirtualProduction/CaptureManager/CaptureManagerCore/Source/DataIngestCore/Public/IngestCaptureData.h
+++ b/Engine/Plugins/VirtualProduction/CaptureManager/CaptureManagerCore/Source/DataIngestCore/Public/IngestCaptureData.h
@@ -55,6 +55,9 @@ struct FIngestCaptureData
 	TArray<FVideo> Depth;
 	TArray<FAudio> Audio;
 	TArray<FCalibration> Calibration;
+
+	/** Converts all relative paths to absolute using InBasePath. */
+	UE_API void MakePathsAbsolute(const FString& InBasePath);
 };
 
 namespace UE::CaptureManager::IngestCaptureData
diff --git a/Engine/Plugins/VirtualProduction/CaptureManager/CaptureManagerDevices/Source/CPSLiveLinkDevice/CPSLiveLinkDevice.Build.cs b/Engine/Plugins/VirtualProduction/CaptureManager/CaptureManagerDevices/Source/CPSLiveLinkDevice/CPSLiveLinkDevice.Build.cs
--- a/Engine/Plugins/VirtualProduction/CaptureManager/CaptureManagerDevices/Source/CPSLiveLinkDevice/CPSLiveLinkDevice.Build.cs
+++ b/Engine/Plugins/VirtualProduction/CaptureManager/CaptureManagerDevices/Source/CPSLiveLinkDevice/CPSLiveLinkDevice.Build.cs
@@ -29,12 +29,11 @@ public class CPSLiveLinkDevice : ModuleRules
 			"Media",
 			"LiveLinkHubCaptureMessaging",
 			"LiveLinkHub",
-			"LiveLinkFaceMetadata",
 			"RTSPMedia",
+			"CaptureMetadataExtraction",
 			"ToolWidgets",
 			"Slate",
 			"SlateCore",
-			"StereoCameraMetadata",
 			"CaptureManagerTakeMetadata",
 			"CaptureManagerSettings",
 			"CoreUObject",
diff --git a/Engine/Plugins/VirtualProduction/CaptureManager/CaptureManagerDevices/Source/CPSLiveLinkDevice/Private/LiveLinkFaceDevice.cpp b/Engine/Plugins/VirtualProduction/CaptureManager/CaptureManagerDevices/Source/CPSLiveLinkDevice/Private/LiveLinkFaceDevice.cpp
--- a/Engine/Plugins/VirtualProduction/CaptureManager/CaptureManagerDevices/Source/CPSLiveLinkDevice/Private/LiveLinkFaceDevice.cpp
+++ b/Engine/Plugins/VirtualProduction/CaptureManager/CaptureManagerDevices/Source/CPSLiveLinkDevice/Private/LiveLinkFaceDevice.cpp
@@ -13,10 +13,8 @@
 #include "Protocol/CPSFileStream.h"
 
 #include "LiveLinkFaceMetadata.h"
+#include "LiveLinkFaceMetadataExtractor.h"
 #include "Settings/CaptureManagerSettings.h"
-#include "StereoCameraMetadataParseUtils.h"
-#include "Utils/CaptureExtractTimecode.h"
-#include "Utils/ParseTakeUtils.h"
 
 #include "RtspMediaSource.h"
 
@@ -495,9 +493,19 @@ void ULiveLinkFaceDevice::OnExportFinished(UE::CaptureManager::TProtocolResult<v
 
 	if (InResult.HasValue())
 	{
-		// Blocking function
-		TOptional<FTakeMetadata> MaybeTakeMetadata = ParseTake(DownloadedStorage, InTakeName);
-		if (MaybeTakeMetadata.IsSet())
+		FExtractionConfig ExtractionConfig;
+		if (const UCaptureManagerSettings* CaptureManagerSettings = GetDefault<UCaptureManagerSettings>())
+		{
+			if (CaptureManagerSettings->bEnableThirdPartyEncoder && !CaptureManagerSettings->ThirdPartyEncoder.FilePath.IsEmpty())
+			{
+				ExtractionConfig.bUseFFprobe = CaptureManagerSettings->bEnableThirdPartyEncoder;
+				ExtractionConfig.FFmpegPath = CaptureManagerSettings->ThirdPartyEncoder.FilePath;
+			}
+		}
+
+		TValueOrError<FTakeMetadata, ELiveLinkFaceExtractionError> ExtractResult = ExtractLiveLinkFaceMetadata(DownloadedTake, ExtractionConfig);
+
+		if (ExtractResult.HasValue())
 		{
 			FTakeId TakeId = InProcessHandle->GetTakeId();
 
@@ -505,11 +513,7 @@ void ULiveLinkFaceDevice::OnExportFinished(UE::CaptureManager::TProtocolResult<v
 			DownloadedTakes.Add(TakeId, DownloadedTake);
 			Lock.Unlock();
 
-			FTakeMetadata TakeMetadata = MaybeTakeMetadata.GetValue();
-			
-			ExtractTimecodeIfNotSet(TakeMetadata);
-			
-			UpdateTake(TakeId, MoveTemp(TakeMetadata));
+			UpdateTake(TakeId, ExtractResult.StealValue());
 
 			ExecuteProcessFinishedReporter(InProcessHandle.Get(), MakeValue());
 		}
@@ -528,44 +532,6 @@ void ULiveLinkFaceDevice::OnExportFinished(UE::CaptureManager::TProtocolResult<v
 	}
 }
 
-TOptional<FTakeMetadata> ULiveLinkFaceDevice::ParseTake(const FString& InTakeDirectory, const FString& InTakeName)
-{
-	using namespace UE::CaptureManager;
-
-	const FString TakeDirectory = InTakeDirectory;
-	const FString TakePath = FPaths::Combine(TakeDirectory, InTakeName);
-
-	TArray<FString> TakeFiles;
-	IFileManager::Get().FindFiles(TakeFiles, *TakePath, *FTakeMetadata::FileExtension);
-	if (!TakeFiles.IsEmpty())
-	{
-		check(TakeFiles.Num() == 1);
-
-		FTakeMetadataParser TakeMetadataParser;
-		TValueOrError<FTakeMetadata, FTakeMetadataParserError> TakeMetadataResult = TakeMetadataParser.Parse(TakePath / TakeFiles[0]);
-		if (!TakeMetadataResult.HasError())
-		{
-			return TakeMetadataResult.StealValue();
-		}
-	}
-
-	TArray<FText> ValidationFailures;
-	TOptional<FTakeMetadata> ParseOldMetadataResult = LiveLinkMetadata::ParseOldLiveLinkTakeMetadata(TakePath, ValidationFailures);
-
-	if (ParseOldMetadataResult.IsSet())
-	{
-		return ParseOldMetadataResult.GetValue();
-	}
-
-	ParseOldMetadataResult = StereoCameraMetadata::ParseOldStereoCameraMetadata(TakePath, ValidationFailures);
-
-	if (ParseOldMetadataResult.IsSet())
-	{
-		return ParseOldMetadataResult.GetValue();
-	}
-
-	return {};
-}
 
 void ULiveLinkFaceDevice::FetchPreIngestFiles(TMap<FString, UE::CaptureManager::FTakeId> InNameToIdMap)
 {
@@ -640,48 +606,3 @@ void ULiveLinkFaceDevice::RemoveDownloadedTakeData(const UE::CaptureManager::FTa
 	IFileManager::Get().DeleteDirectory(**DownloadedTake, false, true);
 }
 
-void ULiveLinkFaceDevice::ExtractTimecodeIfNotSet(FTakeMetadata& InOutTakeMetadata)
-{
-	using namespace UE::CaptureManager;
-
-	bool bVideoFrameRateSet = false;
-	FFrameRate VideoFrameRate;
-	for (FTakeMetadata::FVideo& Video : InOutTakeMetadata.Video)
-	{
-		if (!Video.TimecodeStart.IsSet())
-		{
-			FCaptureExtractVideoInfo::FResult ExtractorOpt = FCaptureExtractVideoInfo::Create(Video.Path);
-			if (ExtractorOpt.IsValid())
-			{
-				FCaptureExtractVideoInfo Extractor = ExtractorOpt.StealValue();
-				Video.TimecodeStart = Extractor.GetTimecode().ToString();
-				if (FMath::IsNearlyZero(Video.FrameRate))
-				{
-					Video.FrameRate = static_cast<float>(Extractor.GetFrameRate().AsDecimal());
-				}
-			}
-		}
-
-		if (!bVideoFrameRateSet)
-		{
-			VideoFrameRate = UE::CaptureManager::ParseFrameRate(Video.FrameRate);
-		}
-	}
-
-	for (FTakeMetadata::FAudio& Audio : InOutTakeMetadata.Audio)
-	{
-		if (!Audio.TimecodeStart.IsSet() && !Audio.TimecodeRate.IsSet())
-		{
-			TSharedPtr<FCaptureExtractAudioTimecode> Extractor = MakeShareable(new FCaptureExtractAudioTimecode(Audio.Path));
-
-			// The video frame rate will be used to calculate the timecode rate if the timecode rate cannot be extracted from the audio file
-			FCaptureExtractAudioTimecode::FTimecodeInfoResult TimecodeInfoResult = Extractor->Extract(VideoFrameRate);
-			if (TimecodeInfoResult.IsValid())
-			{
-				FTimecodeInfo TimecodeInfo = TimecodeInfoResult.GetValue();
-				Audio.TimecodeStart = TimecodeInfo.Timecode.ToString();
-				Audio.TimecodeRate = static_cast<float>(TimecodeInfo.TimecodeRate.AsDecimal());
-			}
-		}
-	}
-}
diff --git a/Engine/Plugins/VirtualProduction/CaptureManager/CaptureManagerDevices/Source/CPSLiveLinkDevice/Public/LiveLinkFaceDevice.h b/Engine/Plugins/VirtualProduction/CaptureManager/CaptureManagerDevices/Source/CPSLiveLinkDevice/Public/LiveLinkFaceDevice.h
--- a/Engine/Plugins/VirtualProduction/CaptureManager/CaptureManagerDevices/Source/CPSLiveLinkDevice/Public/LiveLinkFaceDevice.h
+++ b/Engine/Plugins/VirtualProduction/CaptureManager/CaptureManagerDevices/Source/CPSLiveLinkDevice/Public/LiveLinkFaceDevice.h
@@ -102,15 +102,12 @@ private:
 						  TStrongObjectPtr<const UIngestCapability_ProcessHandle> InProcessHandle,
 						  TStrongObjectPtr<const UIngestCapability_Options> InIngestOptions);
 
-	TOptional<FTakeMetadata> ParseTake(const FString& InTakeDirectory, const FString& InTakeName);
 	FTakeMetadata ParseTakeMetadata(const UE::CaptureManager::FGetTakeMetadataResponse::FTakeObject& InTake);
 
 	void FetchPreIngestFiles(TMap<FString, UE::CaptureManager::FTakeId> InNameToIdMap);
 
 	void RemoveDownloadedTakeData(const UE::CaptureManager::FTakeId InTakeId);
 
-	void ExtractTimecodeIfNotSet(FTakeMetadata& InOutTakeMetadata);
-
 	TSharedPtr<UE::CaptureManager::FCPSDevice> Device;
 
 	mutable FCriticalSection DownloadedTakesMutex;
diff --git a/Engine/Plugins/VirtualProduction/CaptureManager/CaptureManagerDevices/Source/MonoVideoIngestDevice/MonoVideoIngestDevice.Build.cs b/Engine/Plugins/VirtualProduction/CaptureManager/CaptureManagerDevices/Source/MonoVideoIngestDevice/MonoVideoIngestDevice.Build.cs
--- a/Engine/Plugins/VirtualProduction/CaptureManager/CaptureManagerDevices/Source/MonoVideoIngestDevice/MonoVideoIngestDevice.Build.cs
+++ b/Engine/Plugins/VirtualProduction/CaptureManager/CaptureManagerDevices/Source/MonoVideoIngestDevice/MonoVideoIngestDevice.Build.cs
@@ -34,7 +34,8 @@ public class MonoVideoIngestDevice : ModuleRules
 			"CaptureManagerSettings",
 			"CoreUObject",
 			"UnrealEd",
-			"VideoLiveLinkDeviceCommon"
+			"VideoLiveLinkDeviceCommon",
+			"CaptureMetadataExtraction"
 		});
 
 		CppCompileWarningSettings.UnsafeTypeCastWarningLevel = WarningLevel.Error;
diff --git a/Engine/Plugins/VirtualProduction/CaptureManager/CaptureManagerDevices/Source/MonoVideoIngestDevice/Private/MonoVideoIngestDevice.cpp b/Engine/Plugins/VirtualProduction/CaptureManager/CaptureManagerDevices/Source/MonoVideoIngestDevice/Private/MonoVideoIngestDevice.cpp
--- a/Engine/Plugins/VirtualProduction/CaptureManager/CaptureManagerDevices/Source/MonoVideoIngestDevice/Private/MonoVideoIngestDevice.cpp
+++ b/Engine/Plugins/VirtualProduction/CaptureManager/CaptureManagerDevices/Source/MonoVideoIngestDevice/Private/MonoVideoIngestDevice.cpp
@@ -3,13 +3,11 @@
 #include "MonoVideoIngestDevice.h"
 
 #include "HAL/FileManager.h"
+#include "MonoVideoMetadataExtractor.h"
 #include "Settings/CaptureManagerSettings.h"
-#include "ImageUtils.h"
 
-#include "Misc/MessageDialog.h"
 #include "Utils/TakeDiscoveryExpressionParser.h"
 #include "Utils/VideoDeviceThumbnailExtractor.h"
-#include "Utils/CaptureExtractTimecode.h"
 
 #include "VideoLiveLinkDeviceLog.h"
 
@@ -239,10 +237,6 @@ TOptional<FTakeMetadata> UMonoVideoIngestDevice::ReadTake(const FString& InCurre
 {
 	using namespace UE::CaptureManager;
 
-	FTakeMetadata TakeMetadata;
-
-	FFileStatData FileData = IFileManager::Get().GetStatData(*InCurrentTakeFile);
-
 	FString FileNameFormat = GetSettings()->VideoDiscoveryExpression.Value;
 	FString FileName = FPaths::GetBaseFilename(InCurrentTakeFile);
 
@@ -253,9 +247,7 @@ TOptional<FTakeMetadata> UMonoVideoIngestDevice::ReadTake(const FString& InCurre
 	if (FileNameFormat != "<Auto>")
 	{
 		FTakeDiscoveryExpressionParser TokenParser(FileNameFormat, FileName, UE::MonoVideoLiveLinkDevice::Private::Delimiters);
-		bool bParse = TokenParser.Parse();
-
-		if (!bParse)
+		if (!TokenParser.Parse())
 		{
 			UE_LOGF(LogVideoLiveLinkDevice, Warning, "Failed to match the specified format (%ls) with the video file (%ls)", *(FileNameFormat), *InCurrentTakeFile);
 			return {};
@@ -266,11 +258,6 @@ TOptional<FTakeMetadata> UMonoVideoIngestDevice::ReadTake(const FString& InCurre
 		Name = TokenParser.GetName();
 	}
 
-	if (SlateName.IsEmpty())
-	{
-		SlateName = MoveTemp(FileName);
-	}
-
 	if (TakeNumber == INDEX_NONE)
 	{
 		TakeNumber = 1;
@@ -281,85 +268,56 @@ TOptional<FTakeMetadata> UMonoVideoIngestDevice::ReadTake(const FString& InCurre
 		Name = TEXT("video");
 	}
 
-	TakeMetadata.Version.Major = 4;
-	TakeMetadata.Version.Minor = 1;
+	FExtractionConfig ExtractionConfig;
+	if (const UCaptureManagerSettings* CaptureManagerSettings = GetDefault<UCaptureManagerSettings>())
+	{
+		if (CaptureManagerSettings->bEnableThirdPartyEncoder && !CaptureManagerSettings->ThirdPartyEncoder.FilePath.IsEmpty())
+		{
+			ExtractionConfig.bUseFFprobe = CaptureManagerSettings->bEnableThirdPartyEncoder;
+			ExtractionConfig.FFmpegPath = CaptureManagerSettings->ThirdPartyEncoder.FilePath;
+		}
+	}
 
-	TakeMetadata.Slate = MoveTemp(SlateName);
-	TakeMetadata.TakeNumber = TakeNumber;
+	FMonoVideoDescriptor Descriptor;
+	Descriptor.VideoFilePath = InCurrentTakeFile;
+	Descriptor.Slate = SlateName;
+	Descriptor.TakeNumber = TakeNumber;
 
-	// Create new guid
-	TakeMetadata.UniqueId = FGuid::NewGuid().ToString(EGuidFormats::DigitsWithHyphensLower);
-	TakeMetadata.DateTime = FileData.CreationTime;
-	TakeMetadata.Device.Model = TEXT("MonoVideo");
-	
-	FTakeMetadata::FVideo Video;
-	FVideoDeviceThumbnailExtractor ThumbnailExtractor;
-	TOptional<FTakeThumbnailData::FRawImage> RawImageOpt = ThumbnailExtractor.ExtractThumbnail(InCurrentTakeFile);
+	TValueOrError<FTakeMetadata, EMonoVideoExtractionError> Result = ExtractMonoVideoMetadata(MoveTemp(Descriptor), ExtractionConfig);
 
-	if (RawImageOpt.IsSet())
+	if (Result.HasError())
 	{
-		TakeMetadata.Thumbnail = MoveTemp(RawImageOpt.GetValue());
+		UE_LOGF(LogVideoLiveLinkDevice, Error, "Failed to extract metadata for video file \"%ls\"", *FPaths::GetCleanFilename(InCurrentTakeFile));
+		return {};
 	}
 
-	Video.Name = MoveTemp(Name);
-	Video.Format = FPaths::GetExtension(InCurrentTakeFile);
-	Video.Path = InCurrentTakeFile;
-	Video.PathType = FTakeMetadata::FVideo::EPathType::File;
+	FTakeMetadata TakeMetadata = Result.StealValue();
 
-	FCaptureExtractVideoInfo::FResult ExtractResult = FCaptureExtractVideoInfo::Create(Video.Path);
-
-	if (ExtractResult.HasValue())
+	// The Core extractor derives Video.Name from the file path. Override with the
+	// name resolved from the device's discovery expression parser.
+	if (TakeMetadata.Video.Num() > 0)
 	{
-		FCaptureExtractVideoInfo VideoFileInfo = ExtractResult.StealValue();
-		FString TimecodeStartString = VideoFileInfo.GetTimecode().ToString();
-		Video.TimecodeStart = TimecodeStartString;
-		Video.FrameRate = static_cast<float>(VideoFileInfo.GetFrameRate().AsDecimal());
+		TakeMetadata.Video[0].Name = MoveTemp(Name);
 
-		switch (VideoFileInfo.GetVideoOrientation())
+		if (!TakeMetadata.Video[0].TimecodeStart.IsSet())
 		{
-			case EMediaOrientation::CW90:
-				Video.Orientation = FTakeMetadata::FVideo::EOrientation::CW90;
-				break;
-			case EMediaOrientation::CW180:
-				Video.Orientation = FTakeMetadata::FVideo::EOrientation::CW180;
-				break;
-			case EMediaOrientation::CW270:
-				Video.Orientation = FTakeMetadata::FVideo::EOrientation::CW270;
-				break;
-			default:
-				Video.Orientation = FTakeMetadata::FVideo::EOrientation::Original;
+			UE_LOGF(LogVideoLiveLinkDevice, Warning, "Failed to determine the timecode for the video file \"%ls\".", *FPaths::GetCleanFilename(InCurrentTakeFile));
 		}
 
-		if (VideoFileInfo.ContainsAudio())
+		if (FMath::IsNearlyZero(TakeMetadata.Video[0].FrameRate))
 		{
-			FTakeMetadata::FAudio Audio;
-			Audio.Name = TEXT("audio");
-			Audio.Path = InCurrentTakeFile;
-			Audio.Duration = VideoFileInfo.GetAudioDurationSeconds();
-
-			Audio.TimecodeStart = TimecodeStartString;
-			Audio.TimecodeRate = Video.FrameRate;
-
-			TakeMetadata.Audio.Add(MoveTemp(Audio));
+			UE_LOGF(LogVideoLiveLinkDevice, Warning, "Failed to determine frame rate for video file: \"%ls\"", *FPaths::GetCleanFilename(InCurrentTakeFile));
 		}
 	}
-	else
-	{
-		UE_LOGF(LogVideoLiveLinkDevice, Error, "Failed to read metadata for video file \"%ls\"", *FPaths::GetCleanFilename(InCurrentTakeFile));
-	}
 
-	if (!Video.TimecodeStart.IsSet())
-	{
-		UE_LOGF(LogVideoLiveLinkDevice, Warning, "Failed to determine the timecode for the video file \"%ls\".", *FPaths::GetCleanFilename(InCurrentTakeFile));
-	}
-
-	if (FMath::IsNearlyZero(Video.FrameRate))
+	// Thumbnail extraction is device-specific (not part of the Core extractor)
+	FVideoDeviceThumbnailExtractor ThumbnailExtractor;
+	TOptional<FTakeThumbnailData::FRawImage> RawImageOpt = ThumbnailExtractor.ExtractThumbnail(InCurrentTakeFile);
+	if (RawImageOpt.IsSet())
 	{
-		UE_LOGF(LogVideoLiveLinkDevice, Warning, "Failed to determine frame rate for video file: \"%ls\"", *FPaths::GetCleanFilename(InCurrentTakeFile));
+		TakeMetadata.Thumbnail = MoveTemp(RawImageOpt.GetValue());
 	}
 
-	TakeMetadata.Video.Add(Video);
-
 	return TakeMetadata;
 }
 
diff --git a/Engine/Plugins/VirtualProduction/CaptureManager/CaptureManagerDevices/Source/StereoVideoIngestDevice/Private/StereoVideoIngestDevice.cpp b/Engine/Plugins/VirtualProduction/CaptureManager/CaptureManagerDevices/Source/StereoVideoIngestDevice/Private/StereoVideoIngestDevice.cpp
--- a/Engine/Plugins/VirtualProduction/CaptureManager/CaptureManagerDevices/Source/StereoVideoIngestDevice/Private/StereoVideoIngestDevice.cpp
+++ b/Engine/Plugins/VirtualProduction/CaptureManager/CaptureManagerDevices/Source/StereoVideoIngestDevice/Private/StereoVideoIngestDevice.cpp
@@ -3,18 +3,14 @@
 #include "StereoVideoIngestDevice.h"
 
 #include "HAL/FileManager.h"
-#include "Settings/CaptureManagerSettings.h"
 #include "ImageUtils.h"
-#include "CaptureManagerMediaRWModule.h"
-#include "Asset/CaptureAssetSanitization.h"
 #include "Misc/FileHelper.h"
-#include "Engine/Texture2D.h"
+#include "Settings/CaptureManagerSettings.h"
+#include "StereoVideoMetadataExtractor.h"
 
-#include "Misc/MessageDialog.h"
+#include "Asset/CaptureAssetSanitization.h"
 #include "Utils/TakeDiscoveryExpressionParser.h"
 #include "Utils/VideoDeviceThumbnailExtractor.h"
-#include "Utils/CaptureExtractTimecode.h"
-#include "Utils/ParseTakeUtils.h"
 
 #include "VideoLiveLinkDeviceLog.h"
 
@@ -518,202 +514,133 @@ TOptional<FTakeMetadata> UStereoVideoIngestDevice::CreateTakeMetadata(const FTak
 {
 	using namespace UE::CaptureManager;
 
-	FTakeMetadata TakeMetadata;
+	// Filter video and audio components
+	TArray<FTakeWithComponents::Component> VideoComponents = InTakeWithComponents.Components.FilterByPredicate([](const FTakeWithComponents::Component& Component)
+	{
+		return Component.Type == ETakeComponentType::VIDEO || Component.Type == ETakeComponentType::IMAGE_SEQUENCE;
+	});
 
-	int32 TakeNumber = InTakeWithComponents.TakeNumber;
+	check(VideoComponents.Num() == UE::StereoVideoLiveLinkDevice::Private::MaximumVideoFilesCountPerTake);
 
+	TArray<FTakeWithComponents::Component> AudioComponents = InTakeWithComponents.Components.FilterByPredicate([](const FTakeWithComponents::Component& Component)
+	{
+		return Component.Type == ETakeComponentType::AUDIO;
+	});
+
+	// Map component paths to extractor parameters
+	const FString& VideoPathA = VideoComponents[0].Path;
+	const FString& VideoPathB = VideoComponents[1].Path;
+
+	int32 TakeNumber = InTakeWithComponents.TakeNumber;
 	if (TakeNumber == INDEX_NONE)
 	{
 		TakeNumber = 1;
 	}
 
-	TakeMetadata.Version.Major = 4;
-	TakeMetadata.Version.Minor = 1;
-
-	TakeMetadata.Slate = InTakeWithComponents.SlateName;
-	TakeMetadata.TakeNumber = TakeNumber;
-
-	// Create new guid
-	TakeMetadata.UniqueId = FGuid::NewGuid().ToString(EGuidFormats::DigitsWithHyphensLower);
-
-	FString InCurrentTakeFile = InTakeWithComponents.Components[0].Path;
-	FFileStatData FileData = IFileManager::Get().GetStatData(*InCurrentTakeFile);
-	TakeMetadata.DateTime = FileData.CreationTime;
+	FExtractionConfig ExtractionConfig;
+	if (const UCaptureManagerSettings* CaptureManagerSettings = GetDefault<UCaptureManagerSettings>())
+	{
+		if (CaptureManagerSettings->bEnableThirdPartyEncoder && !CaptureManagerSettings->ThirdPartyEncoder.FilePath.IsEmpty())
+		{
+			ExtractionConfig.bUseFFprobe = CaptureManagerSettings->bEnableThirdPartyEncoder;
+			ExtractionConfig.FFmpegPath = CaptureManagerSettings->ThirdPartyEncoder.FilePath;
+		}
+	}
 
-	TakeMetadata.Device.Model = TEXT("StereoHMC");
+	FStereoVideoDescriptor Descriptor;
+	Descriptor.VideoPathA = VideoPathA;
+	Descriptor.VideoPathB = VideoPathB;
+	for (const FTakeWithComponents::Component& AudioComponent : AudioComponents)
+	{
+		Descriptor.AudioFilePaths.Add(FPaths::ConvertRelativePathToFull(InTakeWithComponents.TakePath, AudioComponent.Path));
+	}
+	Descriptor.Slate = InTakeWithComponents.SlateName;
+	Descriptor.TakeNumber = TakeNumber;
 
+	TValueOrError<FTakeMetadata, EStereoVideoExtractionError> Result = ExtractStereoVideoMetadata(MoveTemp(Descriptor), ExtractionConfig);
 
-	TArray<FTakeWithComponents::Component> VideoComponents = InTakeWithComponents.Components.FilterByPredicate([](const FTakeWithComponents::Component& Component)
+	if (Result.HasError())
 	{
-		return Component.Type == ETakeComponentType::VIDEO || Component.Type == ETakeComponentType::IMAGE_SEQUENCE;
-	});
-
-	check(VideoComponents.Num() == UE::StereoVideoLiveLinkDevice::Private::MaximumVideoFilesCountPerTake);
+		UE_LOGF(LogVideoLiveLinkDevice, Error, "Failed to extract stereo metadata for take at '%ls'", *InTakeWithComponents.TakePath);
+		return {};
+	}
 
-	bool bVideoFrameRateSet = false;
-	FFrameRate VideoFrameRate;
+	FTakeMetadata TakeMetadata = Result.StealValue();
 
-	for (const FTakeWithComponents::Component& VideoComponent : VideoComponents)
+	// The Core extractor sets generic Video/Audio names from filesystem paths.
+	// Override with names from the device's discovery components.
+	for (int32 VideoIndex = 0; VideoIndex < VideoComponents.Num() && VideoIndex < TakeMetadata.Video.Num(); ++VideoIndex)
 	{
-		FTakeMetadata::FVideo Video;
-
-		FString Name = VideoComponent.Name;
+		FString Name = VideoComponents[VideoIndex].Name;
 		if (Name.IsEmpty())
 		{
-			Name = VideoComponent.Path.Mid(InTakeWithComponents.TakePath.Len() + 1);
+			Name = VideoComponents[VideoIndex].Path.Mid(InTakeWithComponents.TakePath.Len() + 1);
 			SanitizePackagePath(Name, '_');
 		}
+		TakeMetadata.Video[VideoIndex].Name = MoveTemp(Name);
 
-		Video.Name = MoveTemp(Name);
-
-		if (VideoComponent.Type == ETakeComponentType::VIDEO)
+		if (VideoComponents[VideoIndex].Type == ETakeComponentType::VIDEO)
 		{
-			Video.Format = FPaths::GetExtension(VideoComponent.Path);
-			Video.PathType = FTakeMetadata::FVideo::EPathType::File;
-
-			if (!TakeMetadata.Thumbnail.GetRawImage().IsSet())
-			{
-				FVideoDeviceThumbnailExtractor ThumbnailExtractor;
-				TOptional<FTakeThumbnailData::FRawImage> RawImageOpt = ThumbnailExtractor.ExtractThumbnail(VideoComponent.Path);
-
-				if (RawImageOpt.IsSet())
-				{
-					TakeMetadata.Thumbnail = MoveTemp(RawImageOpt.GetValue());
-				}
-			}
-
-			Video.Path = VideoComponent.Path;
-
-			FString MediaFilePath = FPaths::ConvertRelativePathToFull(InTakeWithComponents.TakePath, Video.Path);
-
-			FCaptureExtractVideoInfo::FResult ExtractResult = FCaptureExtractVideoInfo::Create(MediaFilePath);
-			if (ExtractResult.IsValid())
-			{
-				FCaptureExtractVideoInfo VideoFileInfo = ExtractResult.StealValue();
-				Video.TimecodeStart = VideoFileInfo.GetTimecode().ToString();			
-				Video.FrameRate = static_cast<float>(VideoFileInfo.GetFrameRate().AsDecimal());
-				VideoFrameRate = VideoFileInfo.GetFrameRate();
-				bVideoFrameRateSet = true;
-
-				switch (VideoFileInfo.GetVideoOrientation())
-				{
-					case EMediaOrientation::CW90:
-						Video.Orientation = FTakeMetadata::FVideo::EOrientation::CW90;
-						break;
-					case EMediaOrientation::CW180:
-						Video.Orientation = FTakeMetadata::FVideo::EOrientation::CW180;
-						break;
-					case EMediaOrientation::CW270:
-						Video.Orientation = FTakeMetadata::FVideo::EOrientation::CW270;
-						break;
-					default:
-						Video.Orientation = FTakeMetadata::FVideo::EOrientation::Original;
-				}
-			}
-
-			if (!Video.TimecodeStart.IsSet())
+			if (!TakeMetadata.Video[VideoIndex].TimecodeStart.IsSet())
 			{
-				UE_LOGF(LogVideoLiveLinkDevice, Warning, "Failed to determine the timecode for the video file %ls.", *VideoComponent.Path);
+				UE_LOGF(LogVideoLiveLinkDevice, Warning, "Failed to determine the timecode for the video file %ls.", *VideoComponents[VideoIndex].Path);
 			}
 
-			if (FMath::IsNearlyZero(Video.FrameRate))
+			if (FMath::IsNearlyZero(TakeMetadata.Video[VideoIndex].FrameRate))
 			{
-				UE_LOGF(LogVideoLiveLinkDevice, Warning, "Failed to determine frame rate for video file: %ls", *VideoComponent.Path);
+				UE_LOGF(LogVideoLiveLinkDevice, Warning, "Failed to determine frame rate for video file: %ls", *VideoComponents[VideoIndex].Path);
 			}
 		}
-		else // ETakeComponentType::IMAGE_SEQUENCE
-		{
-			Video.PathType = FTakeMetadata::FVideo::EPathType::Folder;
-
-			int32 FramesCount = 0;
-			bool bFirstFrame = true;
-
-			bool bIterationResult = IFileManager::Get().IterateDirectoryRecursively(*VideoComponent.Path, [&](const TCHAR* InPath, bool bInIsDirectory)
-				{
-					if (!bInIsDirectory)
-					{
-						if (bFirstFrame)
-						{
-							bFirstFrame = false;
-
-							if (!TakeMetadata.Thumbnail.GetThumbnailData().IsSet())
-							{
-								TArray<uint8> ThumbnailData;
-								if (FFileHelper::LoadFileToArray(ThumbnailData, InPath))
-								{
-									TakeMetadata.Thumbnail = ThumbnailData;
-								}
-							}
-
-							FString Path, Filename, Extension;
-							FPaths::Split(InPath, Path, Filename, Extension);
-							Video.Format = Extension;
-
-							Video.Path = VideoComponent.Path;
-							{
-								FImage Image;
-								FImageUtils::LoadImage(InPath, Image);
-								if (Image.IsImageInfoValid())
-								{
-									Video.FrameWidth = Image.SizeX;
-									Video.FrameHeight = Image.SizeY;
-								}
-							}
-						}
-
-						// Assume found files are image files and increase frame count
-						FramesCount++;
-					}
-
-					return true;
-				});
-
-			Video.FramesCount = FramesCount;
-		}
-
-		TakeMetadata.Video.Add(MoveTemp(Video));
 	}
 
-
-	TArray<FTakeWithComponents::Component> AudioComponents = InTakeWithComponents.Components.FilterByPredicate([](const FTakeWithComponents::Component& Component)
-	{
-		return Component.Type == ETakeComponentType::AUDIO;
-	});
-
+	// Override Audio.Name with the device's component-derived names
 	int32 AudioNameCounter = 1;
-
-	for (const FTakeWithComponents::Component& AudioComponent : AudioComponents)
-	{	
-		FTakeMetadata::FAudio Audio;
-		FString Name = AudioComponent.Name;
-		if (Name.IsEmpty())
+	for (int32 AudioIndex = 0; AudioIndex < AudioComponents.Num() && AudioIndex < TakeMetadata.Audio.Num(); ++AudioIndex)
+	{
+		FString AudioName = AudioComponents[AudioIndex].Name;
+		if (AudioName.IsEmpty())
 		{
-			Name = TEXT("audio") + FString::FromInt(AudioNameCounter++);
+			AudioName = TEXT("audio") + FString::FromInt(AudioNameCounter);
 		}
-		Audio.Name = Name;
-		Audio.Path = AudioComponent.Path;
-		Audio.Duration = 0;
-
-		FString MediaFilePath = FPaths::ConvertRelativePathToFull(InTakeWithComponents.TakePath, Audio.Path);
-		TSharedPtr<FCaptureExtractAudioTimecode> Extractor = MakeShareable(new FCaptureExtractAudioTimecode(MediaFilePath));
+		++AudioNameCounter;
+		TakeMetadata.Audio[AudioIndex].Name = MoveTemp(AudioName);
+	}
 
-		// The video frame rate will be used to calculate the timecode rate if the timecode rate cannot be extracted from the audio file
-		// This is why we need to iterate over the video and audio components independently
-		if (!bVideoFrameRateSet)
+	// Thumbnail extraction is device-specific (not part of the Core extractor)
+	for (int32 VideoIndex = 0; VideoIndex < VideoComponents.Num(); ++VideoIndex)
+	{
+		if (VideoComponents[VideoIndex].Type == ETakeComponentType::VIDEO)
 		{
-			UE_LOGF(LogVideoLiveLinkDevice, Warning, "Target frame rate from the input video(s) not found. This is used to set the timecode rate for the audio. Audio timecode may be set incorrectly as a result.");
+			FVideoDeviceThumbnailExtractor ThumbnailExtractor;
+			TOptional<FTakeThumbnailData::FRawImage> RawImageOpt = ThumbnailExtractor.ExtractThumbnail(VideoComponents[VideoIndex].Path);
+			if (RawImageOpt.IsSet())
+			{
+				TakeMetadata.Thumbnail = MoveTemp(RawImageOpt.GetValue());
+				break;
+			}
 		}
-
-		FCaptureExtractAudioTimecode::FTimecodeInfoResult TimecodeInfoResult = Extractor->Extract(VideoFrameRate);
-		if (TimecodeInfoResult.IsValid())
+		else // IMAGE_SEQUENCE
 		{
-			FTimecodeInfo TimecodeInfo = TimecodeInfoResult.GetValue();
-
-			Audio.TimecodeStart = TimecodeInfo.Timecode.ToString();
-			Audio.TimecodeRate = static_cast<float>(TimecodeInfo.TimecodeRate.AsDecimal());
+			bool bFound = false;
+			IFileManager::Get().IterateDirectoryRecursively(*VideoComponents[VideoIndex].Path, [&](const TCHAR* InPath, bool bInIsDirectory) -> bool
+			{
+				if (!bInIsDirectory)
+				{
+					TArray<uint8> ThumbnailData;
+					if (FFileHelper::LoadFileToArray(ThumbnailData, InPath))
+					{
+						TakeMetadata.Thumbnail = ThumbnailData;
+						bFound = true;
+					}
+					return false;
+				}
+				return true;
+			});
+			if (bFound)
+			{
+				break;
+			}
 		}
-
-		TakeMetadata.Audio.Add(MoveTemp(Audio));
 	}
 
 	return TakeMetadata;
diff --git a/Engine/Plugins/VirtualProduction/CaptureManager/CaptureManagerDevices/Source/StereoVideoIngestDevice/StereoVideoIngestDevice.Build.cs b/Engine/Plugins/VirtualProduction/CaptureManager/CaptureManagerDevices/Source/StereoVideoIngestDevice/StereoVideoIngestDevice.Build.cs
--- a/Engine/Plugins/VirtualProduction/CaptureManager/CaptureManagerDevices/Source/StereoVideoIngestDevice/StereoVideoIngestDevice.Build.cs
+++ b/Engine/Plugins/VirtualProduction/CaptureManager/CaptureManagerDevices/Source/StereoVideoIngestDevice/StereoVideoIngestDevice.Build.cs
@@ -35,7 +35,8 @@ public class StereoVideoIngestDevice : ModuleRules
 			"CaptureManagerSettings",
 			"CoreUObject",
 			"UnrealEd",
-			"VideoLiveLinkDeviceCommon"
+			"VideoLiveLinkDeviceCommon",
+			"CaptureMetadataExtraction"
 		});
 
 		CppCompileWarningSettings.UnsafeTypeCastWarningLevel = WarningLevel.Error;
diff --git a/Engine/Plugins/VirtualProduction/CaptureManager/CaptureManagerDevices/Source/TakeArchiveIngestDevice/Private/TakeArchiveIngestDevice.cpp b/Engine/Plugins/VirtualProduction/CaptureManager/CaptureManagerDevices/Source/TakeArchiveIngestDevice/Private/TakeArchiveIngestDevice.cpp
--- a/Engine/Plugins/VirtualProduction/CaptureManager/CaptureManagerDevices/Source/TakeArchiveIngestDevice/Private/TakeArchiveIngestDevice.cpp
+++ b/Engine/Plugins/VirtualProduction/CaptureManager/CaptureManagerDevices/Source/TakeArchiveIngestDevice/Private/TakeArchiveIngestDevice.cpp
@@ -7,11 +7,9 @@
 
 #include "Utils/TakeArchiveIngestDeviceLog.h"
 
-#include "LiveLinkFaceMetadata.h"
-#include "StereoCameraMetadataParseUtils.h"
+#include "TakeArchiveMetadataExtractor.h"
 
 #include "Utils/IngestLiveLinkDeviceUtils.h"
-#include "Utils/CaptureExtractTimecode.h"
 #include "Utils/ParseTakeUtils.h"
 
 DEFINE_LOG_CATEGORY(LogTakeArchiveIngestDevice);
@@ -112,6 +110,8 @@ void UTakeArchiveIngestDevice::RunConvertAndUploadTake(const UIngestCapability_P
 
 void UTakeArchiveIngestDevice::RunUpdateTakeList(UIngestCapability_UpdateTakeListCallback* InCallback)
 {
+	using namespace UE::CaptureManager;
+
 	RemoveAllTakes();
 
 	FString StoragePath = GetSettings()->TakeDirectory.Path;
@@ -140,6 +140,17 @@ void UTakeArchiveIngestDevice::RunUpdateTakeList(UIngestCapability_UpdateTakeLis
 			return true;
 		});
 
+	FExtractionConfig ExtractionConfig;
+
+	if (const UCaptureManagerSettings* CaptureManagerSettings = GetDefault<UCaptureManagerSettings>())
+	{
+		if (CaptureManagerSettings->bEnableThirdPartyEncoder && !CaptureManagerSettings->ThirdPartyEncoder.FilePath.IsEmpty())
+		{
+			ExtractionConfig.bUseFFprobe = CaptureManagerSettings->bEnableThirdPartyEncoder;
+			ExtractionConfig.FFmpegPath = CaptureManagerSettings->ThirdPartyEncoder.FilePath;
+		}
+	}
+
 	for (const FString& CurrentTakeFile : TakeMetadataFiles)
 	{
 		if (IsUpdateTakeListAbortRequested())
@@ -148,11 +159,11 @@ void UTakeArchiveIngestDevice::RunUpdateTakeList(UIngestCapability_UpdateTakeLis
 			return;
 		}
 
-		TOptional<FTakeMetadata> TakeInfoResult = ReadTake(CurrentTakeFile);
+		TValueOrError<FTakeMetadata, ETakeArchiveExtractionError> TakeMetadata = ExtractTakeArchiveMetadata(CurrentTakeFile, ExtractionConfig);
 
-		if (TakeInfoResult.IsSet())
+		if (TakeMetadata.HasValue())
 		{
-			int32 TakeId = AddTake(TakeInfoResult.GetValue());
+			int32 TakeId = AddTake(TakeMetadata.StealValue());
 			FString CurrentDirectory = FPaths::GetPath(CurrentTakeFile);
 
 			FString TakePath = CurrentDirectory.Right(CurrentDirectory.Len() - StoragePath.Len());
@@ -209,123 +220,3 @@ bool UTakeArchiveIngestDevice::Disconnect_Implementation()
 	return true;
 }
 
-TOptional<FTakeMetadata> UTakeArchiveIngestDevice::ReadTake(const FString& InCurrentTakeFile)
-{
-	using namespace UE::CaptureManager;
-
-	FTakeMetadataParser TakeMetadataParser;
-	TValueOrError<FTakeMetadata, FTakeMetadataParserError> TakeMetadataResult = TakeMetadataParser.Parse(InCurrentTakeFile);
-	FString CurrentDirectory = FPaths::GetPath(InCurrentTakeFile);
-
-	if (TakeMetadataResult.HasError())
-	{
-		FTakeMetadataParserError ParserError = TakeMetadataResult.StealError();
-		FString OriginName = UE::CaptureManager::ErrorOriginToString(ParserError.Origin);
-
-		FString Message =
-			FString::Format(TEXT("Unable to parse take metadata file - {0} (Error origin: {1}): {2}"),
-							{
-								InCurrentTakeFile,
-								ParserError.Message.ToString(),
-								OriginName
-							});
-
-		UE_LOGF(LogTakeArchiveIngestDevice, Warning, "%ls", *Message);
-
-		UE_LOGF(LogTakeArchiveIngestDevice, Display, "Checking backwards compatible take metadata formats");
-
-		TArray<FText> ValidationFailures;
-
-		UE_LOGF(LogTakeArchiveIngestDevice, Display, "Checking directory (%ls) for pre UE 5.6 LiveLink take metadata format", *CurrentDirectory);
-		TOptional<FTakeMetadata> ParseOldMetadataResult = LiveLinkMetadata::ParseOldLiveLinkTakeMetadata(CurrentDirectory, ValidationFailures);
-
-		if (ParseOldMetadataResult.IsSet())
-		{
-			return ParseOldMetadataResult.GetValue();
-		}
-
-		UE_LOGF(LogTakeArchiveIngestDevice, Display, "Checking directory (%ls) for pre UE 5.6 StereoCamera take metadata format", *CurrentDirectory);
-		ParseOldMetadataResult = StereoCameraMetadata::ParseOldStereoCameraMetadata(CurrentDirectory, ValidationFailures);
-
-		if (ParseOldMetadataResult.IsSet())
-		{
-			return ParseOldMetadataResult.GetValue();
-		}
-
-		UE_LOGF(LogTakeArchiveIngestDevice, Error, "Failed to parse take metadata file: %ls", *CurrentDirectory);
-
-		return {};
-	}
-
-	FTakeMetadata& TakeMetadata = TakeMetadataResult.GetValue();
-
-	TOptional<FString> ThumbnailPath = TakeMetadata.Thumbnail.GetThumbnailPath();
-	if (ThumbnailPath.IsSet())
-	{
-		if (FPaths::IsRelative(*ThumbnailPath))
-		{
-			TakeMetadata.Thumbnail = FTakeThumbnailData(FPaths::ConvertRelativePathToFull(CurrentDirectory, *ThumbnailPath));
-		}
-	}
-	
-	bool bVideoFrameRateSet = false;
-	FFrameRate VideoFrameRate;
-	for (FTakeMetadata::FVideo& Video : TakeMetadata.Video)
-	{
-		if (FPaths::IsRelative(*Video.Path))
-		{
-			Video.Path = FPaths::ConvertRelativePathToFull(CurrentDirectory, *Video.Path);
-		}
-
-		// Check if this is a file (and not an image sequence directory)
-		if (FPaths::FileExists(Video.Path) && !Video.TimecodeStart.IsSet())
-		{
-			FCaptureExtractVideoInfo::FResult ExtractorOpt = FCaptureExtractVideoInfo::Create(Video.Path);
-			if (ExtractorOpt.IsValid())
-			{
-				FCaptureExtractVideoInfo Extractor = ExtractorOpt.StealValue();
-				Video.TimecodeStart = Extractor.GetTimecode().ToString();
-				if (FMath::IsNearlyZero(Video.FrameRate))
-				{
-					Video.FrameRate = static_cast<float>(Extractor.GetFrameRate().AsDecimal());
-				}
-			}
-		}
-
-		if (!bVideoFrameRateSet)
-		{
-			VideoFrameRate = UE::CaptureManager::ParseFrameRate(Video.FrameRate);
-		}
-	}
-
-	for (FTakeMetadata::FAudio& Audio : TakeMetadata.Audio)
-	{
-		if (FPaths::IsRelative(*Audio.Path))
-		{
-			Audio.Path = FPaths::ConvertRelativePathToFull(CurrentDirectory, *Audio.Path);
-		}
-
-		TSharedPtr<FCaptureExtractAudioTimecode> Extractor = MakeShareable(new FCaptureExtractAudioTimecode(Audio.Path));
-
-		// The video frame rate will be used to calculate the timecode rate if the timecode rate cannot be extracted from the audio file
-		// If the VideoFrameRate is not set, the timecode rate may be set to the audio sample rate
-		FCaptureExtractAudioTimecode::FTimecodeInfoResult TimecodeInfoResult = Extractor->Extract(VideoFrameRate);
-		if (TimecodeInfoResult.IsValid())
-		{
-			FTimecodeInfo TimecodeInfo = TimecodeInfoResult.GetValue();
-
-			Audio.TimecodeStart = TimecodeInfo.Timecode.ToString();
-			Audio.TimecodeRate = static_cast<float>(TimecodeInfo.TimecodeRate.AsDecimal());
-		}
-	}
-	
-	for (FTakeMetadata::FCalibration& Calib : TakeMetadata.Calibration)
-	{
-		if (FPaths::IsRelative(*Calib.Path))
-		{
-			Calib.Path = FPaths::ConvertRelativePathToFull(CurrentDirectory, *Calib.Path);
-		}
-	}
-
-	return TakeMetadataResult.StealValue();
-}
\ No newline at end of file
diff --git a/Engine/Plugins/VirtualProduction/CaptureManager/CaptureManagerDevices/Source/TakeArchiveIngestDevice/Private/TakeArchiveIngestDevice.h b/Engine/Plugins/VirtualProduction/CaptureManager/CaptureManagerDevices/Source/TakeArchiveIngestDevice/Private/TakeArchiveIngestDevice.h
--- a/Engine/Plugins/VirtualProduction/CaptureManager/CaptureManagerDevices/Source/TakeArchiveIngestDevice/Private/TakeArchiveIngestDevice.h
+++ b/Engine/Plugins/VirtualProduction/CaptureManager/CaptureManagerDevices/Source/TakeArchiveIngestDevice/Private/TakeArchiveIngestDevice.h
@@ -66,7 +66,5 @@ private:
 	virtual bool Disconnect_Implementation() override;
 	//~ End ILiveLinkDeviceCapability_Connection interface
 
-	static TOptional<FTakeMetadata> ReadTake(const FString& InCurrentTakeFile);
-
 	TMap<int32, FString> RelativeTakePaths;
 };
diff --git a/Engine/Plugins/VirtualProduction/CaptureManager/CaptureManagerDevices/Source/TakeArchiveIngestDevice/TakeArchiveIngestDevice.Build.cs b/Engine/Plugins/VirtualProduction/CaptureManager/CaptureManagerDevices/Source/TakeArchiveIngestDevice/TakeArchiveIngestDevice.Build.cs
--- a/Engine/Plugins/VirtualProduction/CaptureManager/CaptureManagerDevices/Source/TakeArchiveIngestDevice/TakeArchiveIngestDevice.Build.cs
+++ b/Engine/Plugins/VirtualProduction/CaptureManager/CaptureManagerDevices/Source/TakeArchiveIngestDevice/TakeArchiveIngestDevice.Build.cs
@@ -28,13 +28,12 @@ public class TakeArchiveIngestDevice : ModuleRules
 			"Media",
 			"LiveLinkHubCaptureMessaging",
 			"LiveLinkHub",
-			"LiveLinkFaceMetadata",
 			"ToolWidgets",
 			"Slate",
 			"SlateCore",
-			"StereoCameraMetadata",
 			"CaptureManagerTakeMetadata",
 			"CaptureManagerSettings",
+			"CaptureMetadataExtraction",
 			"CoreUObject",
 			"UnrealEd"
 		});
diff --git a/Engine/Plugins/VirtualProduction/CaptureManager/CaptureManagerEditor/Source/DataIngestCoreEditor/Internal/CaptureManagerIngestPreparation.cpp b/Engine/Plugins/VirtualProduction/CaptureManager/CaptureManagerEditor/Source/DataIngestCoreEditor/Internal/CaptureManagerIngestPreparation.cpp
new file mode 100644
--- /dev/null
+++ b/Engine/Plugins/VirtualProduction/CaptureManager/CaptureManagerEditor/Source/DataIngestCoreEditor/Internal/CaptureManagerIngestPreparation.cpp
@@ -0,0 +1,419 @@
+// Copyright Epic Games, Inc. All Rights Reserved.
+
+#include "CaptureManagerIngestPreparation.h"
+
+#include "IngestAssetCreator.h"
+
+#include "Asset/CaptureAssetSanitization.h"
+#include "Settings/CaptureManagerEditorSettings.h"
+#include "Settings/CaptureManagerEditorTemplateTokens.h"
+#include "Utils/ParseTakeUtils.h"
+#include "Utils/UnrealCalibrationParser.h"
+#include "CameraCalibration.h"
+
+#include "Engine/Engine.h"
+#include "NamingTokenData.h"
+#include "NamingTokensEngineSubsystem.h"
+
+#include "AssetRegistry/IAssetRegistry.h"
+#include "AssetRegistry/AssetRegistryModule.h"
+#include "Editor.h"
+#include "FileHelpers.h"
+#include "Subsystems/EditorAssetSubsystem.h"
+
+namespace UE::CaptureManager::Private
+{
+UNamingTokensEngineSubsystem* GetTokenSubsystem()
+{
+	return GEngine ? GEngine->GetEngineSubsystem<UNamingTokensEngineSubsystem>() : nullptr;
+}
+
+FNamingTokenFilterArgs MakeBaseTokenArgs(const UCaptureManagerEditorSettings* InSettings)
+{
+	FNamingTokenFilterArgs Args;
+	if (const UCaptureManagerIngestNamingTokens* Tokens = InSettings->GetGeneralNamingTokens())
+	{
+		Args.AdditionalNamespacesToInclude.Add(Tokens->GetNamespace());
+	}
+	return Args;
+}
+} // namespace UE::CaptureManager::Private
+
+namespace UE::CaptureManager
+{
+
+FAssetNamingStrategy::FAssetNamingStrategy(const FIngestCaptureData& InIngestCaptureData) :
+	FAssetNamingStrategy(InIngestCaptureData, TEXT(""), TEXT(""))
+{
+}
+
+FAssetNamingStrategy::FAssetNamingStrategy(const FIngestCaptureData& InIngestCaptureData, FString InImportID, FString InDeviceName)
+{
+	check(IsInGameThread());
+
+	if (const UCaptureManagerEditorSettings* Settings = GetDefault<UCaptureManagerEditorSettings>())
+	{
+		if (const UCaptureManagerIngestNamingTokens* Tokens = Settings->GetGeneralNamingTokens())
+		{
+			ImportNamedArgs.Add(Tokens->GetToken(FString(GeneralTokens::IdKey)).Name, InImportID);
+			ImportNamedArgs.Add(Tokens->GetToken(FString(GeneralTokens::DeviceKey)).Name, InDeviceName);
+			ImportNamedArgs.Add(Tokens->GetToken(FString(GeneralTokens::SlateKey)).Name, InIngestCaptureData.Slate);
+			ImportNamedArgs.Add(Tokens->GetToken(FString(GeneralTokens::TakeKey)).Name, FString::FromInt(InIngestCaptureData.TakeNumber));
+		}
+	}
+}
+
+FString FAssetNamingStrategy::ResolveTemplate(
+	const FString& InTemplate,
+	const FNamingTokenFilterArgs& InTokenArgs,
+	TConstArrayView<const FStringFormatNamedArguments*> InFormatArgSets,
+	FSanitizeFunc InSanitize) const
+{
+	check(IsInGameThread());
+
+	UNamingTokensEngineSubsystem* NamingTokensSubsystem = Private::GetTokenSubsystem();
+	if (!NamingTokensSubsystem)
+	{
+		return {};
+	}
+
+	FString Result = InTemplate;
+	for (const FStringFormatNamedArguments* ArgSet : InFormatArgSets)
+	{
+		Result = FString::Format(*Result, *ArgSet);
+	}
+
+	FNamingTokenResultData TokenResult = NamingTokensSubsystem->EvaluateTokenString(Result, InTokenArgs);
+	Result = TokenResult.EvaluatedText.ToString();
+	InSanitize(Result, TEXT('_'));
+
+	return Result;
+}
+
+FStringFormatNamedArguments FAssetNamingStrategy::BuildVideoNamedArgs(
+	const UCaptureManagerEditorSettings* InSettings,
+	const FIngestCaptureData::FVideo& InVideo,
+	const FFrameRate& InFrameRate)
+{
+	FStringFormatNamedArguments Args;
+	if (const UCaptureManagerVideoNamingTokens* Tokens = InSettings->GetVideoNamingTokens())
+	{
+		Args.Add(Tokens->GetToken(FString(VideoTokens::NameKey)).Name, InVideo.Name);
+		Args.Add(Tokens->GetToken(FString(VideoTokens::FrameRateKey)).Name, FString::Printf(TEXT("%.2f"), InFrameRate.AsDecimal()));
+	}
+	return Args;
+}
+
+FString FAssetNamingStrategy::GetImportFolder() const
+{
+	UCaptureManagerEditorSettings* Settings = GetMutableDefault<UCaptureManagerEditorSettings>();
+	if (!Settings)
+	{
+		return {};
+	}
+
+	FNamingTokenFilterArgs TokenArgs = Private::MakeBaseTokenArgs(Settings);
+
+	const FStringFormatNamedArguments* ArgSets[] = { &ImportNamedArgs };
+	return ResolveTemplate(Settings->GetVerifiedImportDirectory(), TokenArgs, ArgSets, &SanitizePackagePath);
+}
+
+FString FAssetNamingStrategy::GetCaptureDataAssetName() const
+{
+	UCaptureManagerEditorSettings* Settings = GetMutableDefault<UCaptureManagerEditorSettings>();
+	if (!Settings)
+	{
+		return {};
+	}
+
+	FNamingTokenFilterArgs TokenArgs = Private::MakeBaseTokenArgs(Settings);
+
+	const FStringFormatNamedArguments* ArgSets[] = { &ImportNamedArgs };
+	return ResolveTemplate(Settings->CaptureDataAssetName, TokenArgs, ArgSets, &SanitizeAssetName);
+}
+
+FString FAssetNamingStrategy::GetImageSequenceAssetName(const FIngestCaptureData::FVideo& InVideo, const FFrameRate& InFrameRate) const
+{
+	UCaptureManagerEditorSettings* Settings = GetMutableDefault<UCaptureManagerEditorSettings>();
+	if (!Settings)
+	{
+		return {};
+	}
+
+	FNamingTokenFilterArgs TokenArgs = Private::MakeBaseTokenArgs(Settings);
+	if (const TObjectPtr<const UCaptureManagerVideoNamingTokens> Tokens = Settings->GetVideoNamingTokens())
+	{
+		TokenArgs.AdditionalNamespacesToInclude.Add(Tokens->GetNamespace());
+	}
+
+	FStringFormatNamedArguments VideoArgs = BuildVideoNamedArgs(Settings, InVideo, InFrameRate);
+
+	const FStringFormatNamedArguments* ArgSets[] = { &VideoArgs, &ImportNamedArgs };
+	return ResolveTemplate(Settings->ImageSequenceAssetName, TokenArgs, ArgSets, &SanitizeAssetName);
+}
+
+FString FAssetNamingStrategy::GetDepthSequenceAssetName(const FIngestCaptureData::FVideo& InVideo, const FFrameRate& InFrameRate) const
+{
+	UCaptureManagerEditorSettings* Settings = GetMutableDefault<UCaptureManagerEditorSettings>();
+	if (!Settings)
+	{
+		return {};
+	}
+
+	FNamingTokenFilterArgs TokenArgs = Private::MakeBaseTokenArgs(Settings);
+	if (const TObjectPtr<const UCaptureManagerVideoNamingTokens> Tokens = Settings->GetVideoNamingTokens())
+	{
+		TokenArgs.AdditionalNamespacesToInclude.Add(Tokens->GetNamespace());
+	}
+
+	FStringFormatNamedArguments DepthArgs = BuildVideoNamedArgs(Settings, InVideo, InFrameRate);
+
+	const FStringFormatNamedArguments* ArgSets[] = { &DepthArgs, &ImportNamedArgs };
+	return ResolveTemplate(Settings->DepthSequenceAssetName, TokenArgs, ArgSets, &SanitizeAssetName);
+}
+
+FString FAssetNamingStrategy::GetAudioAssetName(const FIngestCaptureData::FAudio& InAudio) const
+{
+	UCaptureManagerEditorSettings* Settings = GetMutableDefault<UCaptureManagerEditorSettings>();
+	if (!Settings)
+	{
+		return {};
+	}
+
+	FNamingTokenFilterArgs TokenArgs = Private::MakeBaseTokenArgs(Settings);
+	FStringFormatNamedArguments AudioArgs;
+	if (const UCaptureManagerAudioNamingTokens* Tokens = Settings->GetAudioNamingTokens())
+	{
+		TokenArgs.AdditionalNamespacesToInclude.Add(Tokens->GetNamespace());
+		AudioArgs.Add(Tokens->GetToken(FString(AudioTokens::NameKey)).Name, InAudio.Name);
+	}
+
+	const FStringFormatNamedArguments* ArgSets[] = { &AudioArgs, &ImportNamedArgs };
+	return ResolveTemplate(Settings->SoundwaveAssetName, TokenArgs, ArgSets, &SanitizeAssetName);
+}
+
+FString FAssetNamingStrategy::GetCalibrationAssetName(const FIngestCaptureData::FCalibration& InCalibration) const
+{
+	UCaptureManagerEditorSettings* Settings = GetMutableDefault<UCaptureManagerEditorSettings>();
+	if (!Settings)
+	{
+		return {};
+	}
+
+	FNamingTokenFilterArgs TokenArgs = Private::MakeBaseTokenArgs(Settings);
+	FStringFormatNamedArguments CalibArgs;
+	if (const UCaptureManagerCalibrationNamingTokens* Tokens = Settings->GetCalibrationNamingTokens())
+	{
+		TokenArgs.AdditionalNamespacesToInclude.Add(Tokens->GetNamespace());
+		CalibArgs.Add(Tokens->GetToken(FString(CalibTokens::NameKey)).Name, InCalibration.Name);
+	}
+
+	const FStringFormatNamedArguments* ArgSets[] = { &CalibArgs, &ImportNamedArgs };
+	return ResolveTemplate(Settings->CalibrationAssetName, TokenArgs, ArgSets, &SanitizeAssetName);
+}
+
+FString FAssetNamingStrategy::GetLensFileAssetName(const FIngestCaptureData::FCalibration& InCalibration, const FCameraCalibration& InCameraCalibration) const
+{
+	UCaptureManagerEditorSettings* Settings = GetMutableDefault<UCaptureManagerEditorSettings>();
+	if (!Settings)
+	{
+		return {};
+	}
+
+	FNamingTokenFilterArgs TokenArgs = Private::MakeBaseTokenArgs(Settings);
+	FStringFormatNamedArguments CalibArgs;
+	if (const UCaptureManagerCalibrationNamingTokens* CalibNamingTokens = Settings->GetCalibrationNamingTokens())
+	{
+		TokenArgs.AdditionalNamespacesToInclude.Add(CalibNamingTokens->GetNamespace());
+		CalibArgs.Add(CalibNamingTokens->GetToken(FString(CalibTokens::NameKey)).Name, InCalibration.Name);
+	}
+
+	FStringFormatNamedArguments LensFileArgs;
+	if (const UCaptureManagerLensFileNamingTokens* LensTokens = Settings->GetLensFileNamingTokens())
+	{
+		LensFileArgs.Add(LensTokens->GetToken(FString(LensFileTokens::CameraNameKey)).Name, InCameraCalibration.CameraId);
+	}
+
+	const FStringFormatNamedArguments* ArgSets[] = { &LensFileArgs, &CalibArgs, &ImportNamedArgs };
+	return ResolveTemplate(Settings->LensFileAssetName, TokenArgs, ArgSets, &SanitizeAssetName);
+}
+
+FCreateAssetsData BuildAssetData(const FIngestCaptureData& InIngestCaptureData, const FAssetNamingStrategy& InNamingStrategy)
+{
+	check(IsInGameThread());
+
+	FCreateAssetsData CreateAssetData;
+
+	CreateAssetData.PackagePath = InNamingStrategy.GetImportFolder();
+	CreateAssetData.CaptureDataAssetName = InNamingStrategy.GetCaptureDataAssetName();
+
+	for (const FIngestCaptureData::FVideo& Video : InIngestCaptureData.Video)
+	{
+		const FFrameRate FrameRate = Video.FrameRate.IsSet() ? ParseFrameRate(Video.FrameRate.GetValue()) : FFrameRate();
+
+		FCreateAssetsData::FImageSequenceData ImageSequenceData;
+		ImageSequenceData.AssetName = InNamingStrategy.GetImageSequenceAssetName(Video, FrameRate);
+		ImageSequenceData.FrameRate = FrameRate;
+		ImageSequenceData.Name = Video.Name;
+		ImageSequenceData.SequenceDirectory = Video.Path;
+		ImageSequenceData.bTimecodePresent = Video.TimecodeStart.IsSet();
+		ImageSequenceData.Timecode = Video.TimecodeStart.IsSet() ? ParseTimecode(Video.TimecodeStart.GetValue()) : FTimecode();
+		ImageSequenceData.TimecodeRate = ImageSequenceData.FrameRate;
+
+		CreateAssetData.ImageSequences.Add(ImageSequenceData);
+	}
+
+	for (const FIngestCaptureData::FVideo& Depth : InIngestCaptureData.Depth)
+	{
+		const FFrameRate FrameRate = Depth.FrameRate.IsSet() ? ParseFrameRate(Depth.FrameRate.GetValue()) : FFrameRate();
+
+		FCreateAssetsData::FImageSequenceData DepthSequenceData;
+		DepthSequenceData.AssetName = InNamingStrategy.GetDepthSequenceAssetName(Depth, FrameRate);
+		DepthSequenceData.FrameRate = FrameRate;
+		DepthSequenceData.Name = Depth.Name;
+		DepthSequenceData.SequenceDirectory = Depth.Path;
+		DepthSequenceData.bTimecodePresent = Depth.TimecodeStart.IsSet();
+		DepthSequenceData.Timecode = Depth.TimecodeStart.IsSet() ? ParseTimecode(Depth.TimecodeStart.GetValue()) : FTimecode();
+		DepthSequenceData.TimecodeRate = DepthSequenceData.FrameRate;
+		CreateAssetData.DepthSequences.Add(DepthSequenceData);
+	}
+
+	for (const FIngestCaptureData::FAudio& Audio : InIngestCaptureData.Audio)
+	{
+		FCreateAssetsData::FAudioData AudioData;
+		AudioData.AssetName = InNamingStrategy.GetAudioAssetName(Audio);
+		AudioData.Name = Audio.Name;
+		AudioData.WAVFile = Audio.Path;
+		AudioData.bTimecodePresent = Audio.TimecodeStart.IsSet();
+		AudioData.Timecode = Audio.TimecodeStart.IsSet() ? ParseTimecode(Audio.TimecodeStart.GetValue()) : FTimecode();
+		AudioData.TimecodeRate = Audio.TimecodeRate.IsSet() ? ParseFrameRate(Audio.TimecodeRate.GetValue()) : FFrameRate();
+
+		CreateAssetData.AudioClips.Add(AudioData);
+	}
+
+	for (const FIngestCaptureData::FCalibration& Calibration : InIngestCaptureData.Calibration)
+	{
+		FCreateAssetsData::FCalibrationData CalibrationData;
+		CalibrationData.AssetName = InNamingStrategy.GetCalibrationAssetName(Calibration);
+
+		FUnrealCalibrationParser::FParseResult Result = FUnrealCalibrationParser::Parse(Calibration.Path);
+
+		if (Result.HasValue())
+		{
+			CalibrationData.CameraCalibrations = Result.StealValue();
+
+			for (const FCameraCalibration& CamCalib : CalibrationData.CameraCalibrations)
+			{
+				CalibrationData.LensFileAssetNames.Add(CamCalib.CameraId, InNamingStrategy.GetLensFileAssetName(Calibration, CamCalib));
+			}
+		}
+
+		CreateAssetData.Calibrations.Add(MoveTemp(CalibrationData));
+	}
+
+	if (!InIngestCaptureData.Video.IsEmpty())
+	{
+		TArray<FFrameNumber> DroppedFrameNumbers;
+		DroppedFrameNumbers.Reserve(InIngestCaptureData.Video[0].DroppedFrames.Num());
+		for (const uint32 DroppedFrameIndex : InIngestCaptureData.Video[0].DroppedFrames)
+		{
+			DroppedFrameNumbers.Add(FFrameNumber(static_cast<int32>(DroppedFrameIndex)));
+		}
+		CreateAssetData.CaptureExcludedFrames = PackIntoFrameRanges(MoveTemp(DroppedFrameNumbers));
+	}
+
+	return CreateAssetData;
+}
+
+FCaptureDataTakeInfo BuildTakeInfo(const FCreateAssetsData& InAssetsData, const FIngestCaptureData& InIngestCaptureData)
+{
+	FCaptureDataTakeInfo TakeInfo;
+	TakeInfo.Name = InAssetsData.CaptureDataAssetName;
+	TakeInfo.DeviceModel = InIngestCaptureData.DeviceModel;
+
+	if (!InIngestCaptureData.Video.IsEmpty())
+	{
+		const FIngestCaptureData::FVideo& FirstVideo = InIngestCaptureData.Video[0];
+
+		if (FirstVideo.FrameRate.IsSet())
+		{
+			TakeInfo.FrameRate = FirstVideo.FrameRate.GetValue();
+		}
+	}
+
+	return TakeInfo;
+}
+
+UFootageCaptureData* CreateFootageCaptureDataAsset_GameThread(const FString& InAssetPath, const FCaptureDataAssetInfo& InResult, const FCaptureDataTakeInfo& InTakeInfo, bool bShouldSave)
+{
+	check(IsInGameThread());
+
+	FString CaptureDataName = InTakeInfo.Name;
+	CaptureDataName.TrimStartAndEndInline();
+	CaptureDataName.ReplaceCharInline(' ', '_');
+
+	if (FIngestAssetCreator::GetAssetIfExists<UFootageCaptureData>(InAssetPath, CaptureDataName))
+	{
+		return nullptr;
+	}
+
+	UFootageCaptureData* CaptureData = FIngestAssetCreator::CreateAsset<UFootageCaptureData>(InAssetPath, CaptureDataName);
+	if (CaptureData)
+	{
+		CaptureData->ImageSequences.Reset();
+		CaptureData->DepthSequences.Reset();
+		CaptureData->CameraCalibrations.Reset();
+		CaptureData->AudioTracks.Reset();
+
+		for (const FCaptureDataAssetInfo::FImageSequence& ImageSequence : InResult.ImageSequences)
+		{
+			CaptureData->ImageSequences.Add(ImageSequence.Asset);
+		}
+
+		for (const FCaptureDataAssetInfo::FImageSequence& DepthSequence : InResult.DepthSequences)
+		{
+			CaptureData->DepthSequences.Add(DepthSequence.Asset);
+		}
+
+		for (const FCaptureDataAssetInfo::FAudio& Audio : InResult.Audios)
+		{
+			CaptureData->AudioTracks.Add(Audio.Asset);
+		}
+
+		for (const FCaptureDataAssetInfo::FCalibration& Calibration : InResult.Calibrations)
+		{
+			CaptureData->CameraCalibrations.Add(Calibration.Asset);
+		}
+
+		CaptureData->Metadata.FrameRate = InTakeInfo.FrameRate;
+		CaptureData->Metadata.DeviceModelName = InTakeInfo.DeviceModel;
+		CaptureData->Metadata.SetDeviceClass(InTakeInfo.DeviceModel);
+		CaptureData->CaptureExcludedFrames = InResult.CaptureExcludedFrames;
+
+		if (bShouldSave)
+		{
+			IAssetRegistry& AssetRegistry = FModuleManager::GetModuleChecked<FAssetRegistryModule>(TEXT("AssetRegistry")).Get();
+
+			TArray<FAssetData> AssetsData;
+			AssetRegistry.GetAssetsByPath(FName{ *InAssetPath }, AssetsData, true, false);
+
+			TArray<UPackage*> Packages;
+			for (const FAssetData& AssetData : AssetsData)
+			{
+				if (UObject* Asset = AssetData.GetAsset())
+				{
+					Packages.AddUnique(Asset->GetPackage());
+				}
+			}
+
+			if (!Packages.IsEmpty())
+			{
+				UEditorLoadingAndSavingUtils::SavePackages(Packages, true);
+			}
+		}
+	}
+
+	return CaptureData;
+}
+
+}
diff --git a/Engine/Plugins/VirtualProduction/CaptureManager/CaptureManagerEditor/Source/DataIngestCoreEditor/Internal/CaptureManagerIngestPreparation.h b/Engine/Plugins/VirtualProduction/CaptureManager/CaptureManagerEditor/Source/DataIngestCoreEditor/Internal/CaptureManagerIngestPreparation.h
new file mode 100644
--- /dev/null
+++ b/Engine/Plugins/VirtualProduction/CaptureManager/CaptureManagerEditor/Source/DataIngestCoreEditor/Internal/CaptureManagerIngestPreparation.h
@@ -0,0 +1,67 @@
+// Copyright Epic Games, Inc. All Rights Reserved.
+
+#pragma once
+
+#include "IngestAssetCreator.h"
+#include "IngestCaptureData.h"
+
+#include "Misc/FrameRate.h"
+
+#define UE_API DATAINGESTCOREEDITOR_API
+
+struct FCameraCalibration;
+struct FNamingTokenFilterArgs;
+class UCaptureManagerEditorSettings;
+
+namespace UE::CaptureManager
+{
+
+struct FCaptureDataTakeInfo
+{
+	FString Name;
+	double FrameRate = 0.0;
+	FString DeviceModel;
+};
+
+class UE_API FAssetNamingStrategy
+{
+public:
+	FAssetNamingStrategy(const FIngestCaptureData& InIngestCaptureData);
+	FAssetNamingStrategy(const FIngestCaptureData& InIngestCaptureData, FString InImportID, FString InDeviceName);
+
+	FString GetImportFolder() const;
+	FString GetCaptureDataAssetName() const;
+	FString GetImageSequenceAssetName(const FIngestCaptureData::FVideo& InVideo, const FFrameRate& InFrameRate) const;
+	FString GetDepthSequenceAssetName(const FIngestCaptureData::FVideo& InVideo, const FFrameRate& InFrameRate) const;
+	FString GetAudioAssetName(const FIngestCaptureData::FAudio& InAudio) const;
+	FString GetCalibrationAssetName(const FIngestCaptureData::FCalibration& InCalibration) const;
+	FString GetLensFileAssetName(const FIngestCaptureData::FCalibration& InCalibration, const FCameraCalibration& InCameraCalibration) const;
+
+private:
+	using FSanitizeFunc = void(*)(FString&, FString::ElementType);
+
+	FString ResolveTemplate(
+		const FString& InTemplate,
+		const FNamingTokenFilterArgs& InTokenArgs,
+		TConstArrayView<const FStringFormatNamedArguments*> InFormatArgSets,
+		FSanitizeFunc InSanitize) const;
+
+	static FStringFormatNamedArguments BuildVideoNamedArgs(
+		const UCaptureManagerEditorSettings* InSettings,
+		const FIngestCaptureData::FVideo& InVideo,
+		const FFrameRate& InFrameRate);
+
+	FStringFormatNamedArguments ImportNamedArgs;
+};
+
+UE_API FCreateAssetsData BuildAssetData(const FIngestCaptureData& InIngestCaptureData, const FAssetNamingStrategy& InNamingStrategy);
+
+/** Populates a FCaptureDataTakeInfo from the asset name, device model, and first video entry of the ingest data. */
+UE_API FCaptureDataTakeInfo BuildTakeInfo(const FCreateAssetsData& InAssetsData, const FIngestCaptureData& InIngestCaptureData);
+
+/** Creates a UFootageCaptureData asset from prepared asset info and take metadata. Returns nullptr if the asset already exists. */
+UE_API UFootageCaptureData* CreateFootageCaptureDataAsset_GameThread(const FString& InAssetPath, const FCaptureDataAssetInfo& InResult, const FCaptureDataTakeInfo& InTakeInfo, bool bShouldSave = true);
+
+}
+
+#undef UE_API
diff --git a/Engine/Plugins/VirtualProduction/CaptureManager/CaptureManagerEditor/Source/DataIngestCoreEditor/Private/IngestAssetCreator.cpp b/Engine/Plugins/VirtualProduction/CaptureManager/CaptureManagerEditor/Source/DataIngestCoreEditor/Private/IngestAssetCreator.cpp
--- a/Engine/Plugins/VirtualProduction/CaptureManager/CaptureManagerEditor/Source/DataIngestCoreEditor/Private/IngestAssetCreator.cpp
+++ b/Engine/Plugins/VirtualProduction/CaptureManager/CaptureManagerEditor/Source/DataIngestCoreEditor/Private/IngestAssetCreator.cpp
@@ -7,6 +7,7 @@
 #include "AssetToolsModule.h"
 #include "AssetRegistry/IAssetRegistry.h"
 #include "AssetRegistry/AssetRegistryModule.h"
+#include "FileHelpers.h"
 #include "ImgMediaSource.h"
 #include "Sound/SoundWave.h"
 #include "AssetImportTask.h"
@@ -700,7 +701,7 @@ FIngestAssetCreator::FAssetCreationResult FIngestAssetCreator::CheckTakeAssets(c
 	return MakeValue();
 }
 
-void FIngestAssetCreator::RemoveAssetsByPath(const FString& InPackagePath)
+bool FIngestAssetCreator::RemoveAssetsByPath(const FString& InPackagePath)
 {
 	IAssetRegistry& AssetRegistry = FModuleManager::GetModuleChecked<FAssetRegistryModule>(TEXT("AssetRegistry")).Get();
 
@@ -709,13 +710,19 @@ void FIngestAssetCreator::RemoveAssetsByPath(const FString& InPackagePath)
 
 	if (AssetData.IsEmpty())
 	{
-		return;
+		return true;
 	}
 
-	if (ObjectTools::DeleteAssets(AssetData, false) != AssetData.Num())
+	const int32 DeletedCount = ObjectTools::DeleteAssets(AssetData, false);
+	if (DeletedCount != AssetData.Num())
 	{
-		UE_LOGF(LogIngestAssetCreator, Warning, "Not all assets are deleted");
+		UE_LOGF(LogIngestAssetCreator, Warning, "RemoveAssetsByPath: deleted %d of %d assets at '%ls'",
+			DeletedCount, AssetData.Num(), *InPackagePath
+		);
+		return false;
 	}
+
+	return true;
 }
 
 UObject* FIngestAssetCreator::GetAssetIfExists(const FString& InTargetPackagePath, const FString& InAssetName)
diff --git a/Engine/Plugins/VirtualProduction/CaptureManager/CaptureManagerEditor/Source/DataIngestCoreEditor/Public/IngestAssetCreator.h b/Engine/Plugins/VirtualProduction/CaptureManager/CaptureManagerEditor/Source/DataIngestCoreEditor/Public/IngestAssetCreator.h
--- a/Engine/Plugins/VirtualProduction/CaptureManager/CaptureManagerEditor/Source/DataIngestCoreEditor/Public/IngestAssetCreator.h
+++ b/Engine/Plugins/VirtualProduction/CaptureManager/CaptureManagerEditor/Source/DataIngestCoreEditor/Public/IngestAssetCreator.h
@@ -66,7 +66,7 @@ struct FCreateAssetsData
 	FString CaptureDataAssetName;
 
 	/** Unique identifier for this take. */
-	int32 TakeId;
+	int32 TakeId = 0;
 	FString PackagePath;
 	TArray<FImageSequenceData> ImageSequences;
 	TArray<FImageSequenceData> DepthSequences;
@@ -166,9 +166,9 @@ private:
 												 const TObjectPtr<UAssetImportTask>& InAssetImportTask,
 												 FCaptureDataAssetInfo& OutTake);
 
-
-	/** Deletes all assets at specified package path. */
-	static UE_API void RemoveAssetsByPath(const FString& InPackagePath);
+	/** Deletes all assets at the specified package path.
+	 *  Returns true if all assets were deleted (or none existed), false if any deletion failed. */
+	static UE_API bool RemoveAssetsByPath(const FString& InPackagePath);
 
 	/** Converts timecode to string. */
 	static UE_API FString CreateTimecodeString(FTimecode InTimecode, FFrameRate InFrameRate);
@@ -196,7 +196,6 @@ private:
 	/** Checks validity of assets in the take. */
 	static UE_API FAssetCreationResult CheckTakeAssets(const FCaptureDataAssetInfo& InTake);
 
-
 	static UE_API UObject* GetOrCreateAsset(const FString& InTargetPackagePath, const FString& InAssetName, UClass* InClass);
 	static UE_API UObject* CreateAsset(const FString& InTargetPackagePath, const FString& InAssetName, UClass* InClass);
 
diff --git a/Engine/Plugins/VirtualProduction/CaptureManager/CaptureManagerEditor/Source/LiveLinkHubWorkerManager/Private/IngestProcess/IngestCaptureDataProcess.cpp b/Engine/Plugins/VirtualProduction/CaptureManager/CaptureManagerEditor/Source/LiveLinkHubWorkerManager/Private/IngestProcess/IngestCaptureDataProcess.cpp
--- a/Engine/Plugins/VirtualProduction/CaptureManager/CaptureManagerEditor/Source/LiveLinkHubWorkerManager/Private/IngestProcess/IngestCaptureDataProcess.cpp
+++ b/Engine/Plugins/VirtualProduction/CaptureManager/CaptureManagerEditor/Source/LiveLinkHubWorkerManager/Private/IngestProcess/IngestCaptureDataProcess.cpp
@@ -2,6 +2,7 @@
 
 #include "IngestCaptureDataProcess.h"
 
+#include "CaptureManagerIngestPreparation.h"
 #include "HAL/PlatformFileManager.h"
 
 #include "Utils/ParseTakeUtils.h"
@@ -82,7 +83,7 @@ TValueOrError<FIngestProcessResult, FText> FIngestCaptureDataProcess::StartInges
 
 	FIngestCaptureData IngestCaptureData = IngestCaptureDataParseResult.GetValue();
 
-	ConvertPathsToFull(InTakeStoragePath, IngestCaptureData);
+	IngestCaptureData.MakePathsAbsolute(InTakeStoragePath);
 
 	FCreateAssetsData AssetCreationData = PrepareAssetsData(InTakeUploadId, InDeviceName, IngestCaptureData);
 
@@ -90,293 +91,32 @@ TValueOrError<FIngestProcessResult, FText> FIngestCaptureDataProcess::StartInges
 	IngestProcessResult.TakeIngestPackagePath = AssetCreationData.PackagePath;
 	IngestProcessResult.AssetsData.Add(AssetCreationData);
 
-	IngestProcessResult.CaptureDataTakeInfo.Name = AssetCreationData.CaptureDataAssetName;
-	if (!IngestCaptureData.Video.IsEmpty())
-	{
-		const FIngestCaptureData::FVideo& FirstVideo = IngestCaptureData.Video[0];
-		if (FirstVideo.FrameRate.IsSet())
-		{
-			IngestProcessResult.CaptureDataTakeInfo.FrameRate = FirstVideo.FrameRate.GetValue();
-		}
-
-		if (FirstVideo.FrameWidth.IsSet() && FirstVideo.FrameHeight.IsSet())
-		{
-			IngestProcessResult.CaptureDataTakeInfo.Resolution = FIntPoint(FirstVideo.FrameWidth.GetValue(), FirstVideo.FrameHeight.GetValue());
-		}
-
-	}
-	IngestProcessResult.CaptureDataTakeInfo.DeviceModel = IngestCaptureData.DeviceModel;
+	IngestProcessResult.CaptureDataTakeInfo = BuildTakeInfo(AssetCreationData, IngestCaptureData);
 
 	return MakeValue(MoveTemp(IngestProcessResult));
 
 }
 
-UE::CaptureManager::FCreateAssetsData FIngestCaptureDataProcess::PrepareAssetsData(const FGuid& InTakeUploadId,
-																				   const FString& InDeviceName,
-																				   const FIngestCaptureData& InIngestCaptureData)
+UE::CaptureManager::FCreateAssetsData FIngestCaptureDataProcess::PrepareAssetsData(
+	const FGuid& InTakeUploadId,
+	const FString& InDeviceName,
+	const FIngestCaptureData& InIngestCaptureData
+)
 {
 	using namespace UE::CaptureManager;
 
-	UCaptureManagerEditorSettings* Settings = GetMutableDefault<UCaptureManagerEditorSettings>();
-
-	const TObjectPtr<const UCaptureManagerIngestNamingTokens> Tokens = Settings->GetGeneralNamingTokens();
-	check(Tokens);
-
-	FStringFormatNamedArguments ImportNamedArgs;
-	{
-		// Editor tokens
-		ImportNamedArgs.Add(Tokens->GetToken(FString(GeneralTokens::IdKey)).Name, InTakeUploadId.ToString());
-		ImportNamedArgs.Add(Tokens->GetToken(FString(GeneralTokens::DeviceKey)).Name, InDeviceName);
-		ImportNamedArgs.Add(Tokens->GetToken(FString(GeneralTokens::SlateKey)).Name, InIngestCaptureData.Slate);
-		ImportNamedArgs.Add(Tokens->GetToken(FString(GeneralTokens::TakeKey)).Name, FString::FromInt(InIngestCaptureData.TakeNumber));
-	}
-
-	int32 TakeId = 0;
-
-	FCreateAssetsData CreateAssetData;
+	FCreateAssetsData CreateAssetsData;
 
-	CreateAssetData.TakeId = TakeId;
-
-	// Naming tokens subsystem consults asset registry so need to run on the game thread
 	CallOnGameThread(
-		[&CreateAssetData, &InIngestCaptureData, &Settings, &ImportNamedArgs]()
+		[&InIngestCaptureData, &InTakeUploadId, &InDeviceName, &CreateAssetsData]()
 		{
-			UNamingTokensEngineSubsystem* NamingTokensSubsystem = GEngine->GetEngineSubsystem<UNamingTokensEngineSubsystem>();
-			check(NamingTokensSubsystem);
-
-			FNamingTokenFilterArgs NamingTokenArgs;
-			if (const TObjectPtr<const UCaptureManagerIngestNamingTokens> Tokens = Settings->GetGeneralNamingTokens())
-			{
-				NamingTokenArgs.AdditionalNamespacesToInclude.Add(Tokens->GetNamespace());
-			}
-
-			// Evaluate Asset Folder
-			FString VerifiedImportDirectory = Settings->GetVerifiedImportDirectory();
-
-			FString ImportDirectory = FString::Format(*VerifiedImportDirectory, ImportNamedArgs);
-			FNamingTokenResultData ImportDirectoryResult = NamingTokensSubsystem->EvaluateTokenString(ImportDirectory, NamingTokenArgs);
-			FString ProcessedPackagePath = ImportDirectoryResult.EvaluatedText.ToString();
-			SanitizePackagePath(ProcessedPackagePath);
-			CreateAssetData.PackagePath = ProcessedPackagePath;
-
-			// Evaluate Asset Settings
-			FString CaptureDataName = FString::Format(*Settings->CaptureDataAssetName, ImportNamedArgs);
-			FNamingTokenResultData CaptureDataNameResult = NamingTokensSubsystem->EvaluateTokenString(CaptureDataName, NamingTokenArgs);
-			FString ProcessedCaptureDataName = CaptureDataNameResult.EvaluatedText.ToString();
-			SanitizeAssetName(ProcessedCaptureDataName);
-			CreateAssetData.CaptureDataAssetName = ProcessedCaptureDataName;
-
-			FNamingTokenFilterArgs VideoEvaluationArgs = NamingTokenArgs;
-			if (const TObjectPtr<const UCaptureManagerVideoNamingTokens> Tokens = Settings->GetVideoNamingTokens())
-			{
-				VideoEvaluationArgs.AdditionalNamespacesToInclude.Add(Tokens->GetNamespace());
-			}
-
-			for (int32 Index = 0; Index < InIngestCaptureData.Video.Num(); ++Index)
-			{
-				FCreateAssetsData::FImageSequenceData ImageSequenceData;
-
-				const FIngestCaptureData::FVideo& Video = InIngestCaptureData.Video[Index];
-
-				FFrameRate FrameRate = Video.FrameRate.IsSet() ? ParseFrameRate(Video.FrameRate.GetValue()) : FFrameRate();
-
-				FStringFormatNamedArguments VideoNamedArgs;
-				{
-					// Video tokens
-					VideoNamedArgs.Add(Settings->GetVideoNamingTokens()->GetToken(FString(VideoTokens::NameKey)).Name, Video.Name);
-					VideoNamedArgs.Add(Settings->GetVideoNamingTokens()->GetToken(FString(VideoTokens::FrameRateKey)).Name, FString::Printf(TEXT("%.2f"), FrameRate.AsDecimal()));
-				}
-				// Evaluate Video Settings
-				FString ImageSequenceAssetName = FString::Format(*Settings->ImageSequenceAssetName, VideoNamedArgs);
-				ImageSequenceAssetName = FString::Format(*ImageSequenceAssetName, ImportNamedArgs);
-				FNamingTokenResultData ImageSequenceAssetNameResult = NamingTokensSubsystem->EvaluateTokenString(ImageSequenceAssetName, VideoEvaluationArgs);
-				FString ProcessedImageSequenceAssetName = ImageSequenceAssetNameResult.EvaluatedText.ToString();
-				SanitizeAssetName(ProcessedImageSequenceAssetName);
-				ImageSequenceData.AssetName = ProcessedImageSequenceAssetName;
-
-				ImageSequenceData.FrameRate = FrameRate;
-				ImageSequenceData.Name = Video.Name;
-				ImageSequenceData.SequenceDirectory = Video.Path;
-				ImageSequenceData.bTimecodePresent = Video.TimecodeStart.IsSet();
-				ImageSequenceData.Timecode = Video.TimecodeStart.IsSet() ? ParseTimecode(Video.TimecodeStart.GetValue()) : FTimecode();
-				ImageSequenceData.TimecodeRate = ImageSequenceData.FrameRate;
-
-				CreateAssetData.ImageSequences.Add(ImageSequenceData);
-			}
-
-			for (int32 Index = 0; Index < InIngestCaptureData.Depth.Num(); ++Index)
-			{
-				FCreateAssetsData::FImageSequenceData DepthSequenceData;
-
-				const FIngestCaptureData::FVideo& Depth = InIngestCaptureData.Depth[Index];
-
-				FFrameRate FrameRate = Depth.FrameRate.IsSet() ? ParseFrameRate(Depth.FrameRate.GetValue()) : FFrameRate();
-
-				FStringFormatNamedArguments DepthNamedArgs;
-				{
-					// Depth tokens
-					DepthNamedArgs.Add(Settings->GetVideoNamingTokens()->GetToken(FString(VideoTokens::NameKey)).Name, Depth.Name);
-					DepthNamedArgs.Add(Settings->GetVideoNamingTokens()->GetToken(FString(VideoTokens::FrameRateKey)).Name, FString::Printf(TEXT("%.2f"), FrameRate.AsDecimal()));
-				}
-				// Evaluate Depth Settings
-				FString DepthSequenceAssetName = FString::Format(*Settings->DepthSequenceAssetName, DepthNamedArgs);
-				DepthSequenceAssetName = FString::Format(*DepthSequenceAssetName, ImportNamedArgs);
-
-				FNamingTokenResultData DepthSequenceAssetNameResult = NamingTokensSubsystem->EvaluateTokenString(DepthSequenceAssetName, VideoEvaluationArgs);
-				FString ProcessedDepthSequenceAssetName = DepthSequenceAssetNameResult.EvaluatedText.ToString();
-				SanitizeAssetName(ProcessedDepthSequenceAssetName);
-				DepthSequenceData.AssetName = ProcessedDepthSequenceAssetName;
-
-				DepthSequenceData.FrameRate = FrameRate;
-				DepthSequenceData.Name = Depth.Name;
-				DepthSequenceData.SequenceDirectory = Depth.Path;
-				DepthSequenceData.bTimecodePresent = Depth.TimecodeStart.IsSet();
-				DepthSequenceData.Timecode = Depth.TimecodeStart.IsSet() ? ParseTimecode(Depth.TimecodeStart.GetValue()) : FTimecode();
-				DepthSequenceData.TimecodeRate = DepthSequenceData.FrameRate;
-				CreateAssetData.DepthSequences.Add(DepthSequenceData);
-			}
-
-			FNamingTokenFilterArgs AudioEvaluationArgs = NamingTokenArgs;
-			if (const TObjectPtr<const UCaptureManagerAudioNamingTokens> Tokens = Settings->GetAudioNamingTokens())
-			{
-				AudioEvaluationArgs.AdditionalNamespacesToInclude.Add(Tokens->GetNamespace());
-			}
-
-			for (int32 Index = 0; Index < InIngestCaptureData.Audio.Num(); ++Index)
-			{
-				const FIngestCaptureData::FAudio& Audio = InIngestCaptureData.Audio[Index];
-
-				FCreateAssetsData::FAudioData AudioData;
-
-				FStringFormatNamedArguments AudioNamedArgs;
-				{
-					// Audio tokens
-					AudioNamedArgs.Add(Settings->GetAudioNamingTokens()->GetToken(FString(AudioTokens::NameKey)).Name, Audio.Name);
-				}
-				// Evaluate Audio Settings
-				FString AudioAssetName = FString::Format(*Settings->SoundwaveAssetName, AudioNamedArgs);
-				AudioAssetName = FString::Format(*AudioAssetName, ImportNamedArgs);
-
-				FNamingTokenResultData AudioAssetNameResult = NamingTokensSubsystem->EvaluateTokenString(AudioAssetName, AudioEvaluationArgs);
-				FString ProcessedAudioAssetName = AudioAssetNameResult.EvaluatedText.ToString();
-				SanitizeAssetName(ProcessedAudioAssetName);
-				AudioData.AssetName = ProcessedAudioAssetName;
-
-				AudioData.Name = Audio.Name;
-				AudioData.WAVFile = Audio.Path;
-				AudioData.bTimecodePresent = Audio.TimecodeStart.IsSet();
-				AudioData.Timecode = Audio.TimecodeStart.IsSet() ? ParseTimecode(Audio.TimecodeStart.GetValue()) : FTimecode();
-				AudioData.TimecodeRate = Audio.TimecodeRate.IsSet() ? ParseFrameRate(Audio.TimecodeRate.GetValue()) : FFrameRate();
-
-				CreateAssetData.AudioClips.Add(AudioData);
-			}
-
-			FNamingTokenFilterArgs CalibEvaluationArgs = NamingTokenArgs;
-			if (const TObjectPtr<const UCaptureManagerCalibrationNamingTokens> Tokens = Settings->GetCalibrationNamingTokens())
-			{
-				CalibEvaluationArgs.AdditionalNamespacesToInclude.Add(Tokens->GetNamespace());
-			}
-
-			for (const FIngestCaptureData::FCalibration& Calibration : InIngestCaptureData.Calibration)
-			{
-				FCreateAssetsData::FCalibrationData CalibrationData;
-
-				FStringFormatNamedArguments CalibNamedArgs;
-				{
-					// Calibration tokens
-					CalibNamedArgs.Add(Settings->GetCalibrationNamingTokens()->GetToken(FString(CalibTokens::NameKey)).Name, Calibration.Name);
-				}
-
-				// Evaluate Calibration Settings
-				FString CalibAssetName = FString::Format(*Settings->CalibrationAssetName, CalibNamedArgs);
-				CalibAssetName = FString::Format(*CalibAssetName, ImportNamedArgs);
-				
-				FNamingTokenResultData CalibAssetNameResult = NamingTokensSubsystem->EvaluateTokenString(CalibAssetName, CalibEvaluationArgs);
-				FString ProcessedCalibAssetName = CalibAssetNameResult.EvaluatedText.ToString();
-				SanitizeAssetName(ProcessedCalibAssetName);
-				CalibrationData.AssetName = ProcessedCalibAssetName;
-
-				FUnrealCalibrationParser::FParseResult Result = FUnrealCalibrationParser::Parse(Calibration.Path);
-
-				if (Result.HasValue())
-				{
-					CalibrationData.CameraCalibrations = Result.StealValue();
-
-					for (const FCameraCalibration& CamCalib : CalibrationData.CameraCalibrations)
-					{
-						FStringFormatNamedArguments LensFileNamedArgs;
-						{
-							// Lens File tokens
-							LensFileNamedArgs.Add(Settings->GetLensFileNamingTokens()->GetToken(FString(LensFileTokens::CameraNameKey)).Name, CamCalib.CameraId);
-						}
-
-						FString LensFileAssetName = FString::Format(*Settings->LensFileAssetName, LensFileNamedArgs);
-						LensFileAssetName = FString::Format(*LensFileAssetName, CalibNamedArgs);
-						LensFileAssetName = FString::Format(*LensFileAssetName, ImportNamedArgs);
-						
-						FNamingTokenResultData LensFileAssetNameResult = NamingTokensSubsystem->EvaluateTokenString(LensFileAssetName, CalibEvaluationArgs);
-						FString ProcessedLensFileAssetName = LensFileAssetNameResult.EvaluatedText.ToString();
-						SanitizeAssetName(ProcessedLensFileAssetName);
-						CalibrationData.LensFileAssetNames.Add(CamCalib.CameraId, ProcessedLensFileAssetName);
-					}
-				}
-
-				CreateAssetData.Calibrations.Add(MoveTemp(CalibrationData));
-			}
-
-			uint32 LastDroppedFrameIndex = 0;
-			uint32 FirstDroppedFrameIndex = 0;
-			if (!InIngestCaptureData.Video.IsEmpty())
-			{
-				for (const uint32 DroppedFrameIndex : InIngestCaptureData.Video[0].DroppedFrames)
-				{
-					if (FirstDroppedFrameIndex == 0)
-					{
-						FirstDroppedFrameIndex = DroppedFrameIndex;
-					}
-
-					if (DroppedFrameIndex - LastDroppedFrameIndex == 1)
-					{
-						LastDroppedFrameIndex = DroppedFrameIndex;
-					}
-					else
-					{
-						FFrameRange Range;
-						Range.StartFrame = FirstDroppedFrameIndex;
-						Range.EndFrame = LastDroppedFrameIndex;
-						CreateAssetData.CaptureExcludedFrames.Add(MoveTemp(Range));
-
-						FirstDroppedFrameIndex = DroppedFrameIndex;
-						LastDroppedFrameIndex = DroppedFrameIndex;
-					}
-				}
-			}
-		});
-
-	return CreateAssetData;
-}
-
-void FIngestCaptureDataProcess::ConvertPathsToFull(const FString& InTakeStoragePath, FIngestCaptureData& OutIngestCaptureData)
-{
-	for (FIngestCaptureData::FVideo& Video : OutIngestCaptureData.Video)
-	{
-		Video.Path = FPaths::ConvertRelativePathToFull(InTakeStoragePath, Video.Path);
-	}
-
-	for (FIngestCaptureData::FVideo& Depth : OutIngestCaptureData.Depth)
-	{
-		Depth.Path = FPaths::ConvertRelativePathToFull(InTakeStoragePath, Depth.Path);
-	}
-
-	for (FIngestCaptureData::FAudio& Audio : OutIngestCaptureData.Audio)
-	{
-		Audio.Path = FPaths::ConvertRelativePathToFull(InTakeStoragePath, Audio.Path);
-	}
+			// Needs to be on the game thread as we're using naming tokens
+			FAssetNamingStrategy AssetNamingStrategy(InIngestCaptureData, InTakeUploadId.ToString(), InDeviceName);
+			CreateAssetsData = BuildAssetData(InIngestCaptureData, AssetNamingStrategy);
+		}
+	);
 
-	for (FIngestCaptureData::FCalibration& Calibration : OutIngestCaptureData.Calibration)
-	{
-		Calibration.Path = FPaths::ConvertRelativePathToFull(InTakeStoragePath, Calibration.Path);
-	}
+	return CreateAssetsData;
 }
 
 #undef LOCTEXT_NAMESPACE
\ No newline at end of file
diff --git a/Engine/Plugins/VirtualProduction/CaptureManager/CaptureManagerEditor/Source/LiveLinkHubWorkerManager/Private/IngestProcess/IngestCaptureDataProcess.h b/Engine/Plugins/VirtualProduction/CaptureManager/CaptureManagerEditor/Source/LiveLinkHubWorkerManager/Private/IngestProcess/IngestCaptureDataProcess.h
--- a/Engine/Plugins/VirtualProduction/CaptureManager/CaptureManagerEditor/Source/LiveLinkHubWorkerManager/Private/IngestProcess/IngestCaptureDataProcess.h
+++ b/Engine/Plugins/VirtualProduction/CaptureManager/CaptureManagerEditor/Source/LiveLinkHubWorkerManager/Private/IngestProcess/IngestCaptureDataProcess.h
@@ -22,6 +22,4 @@ private:
 	static UE::CaptureManager::FCreateAssetsData PrepareAssetsData(const FGuid& InTakeUploadId,
 																   const FString& InDeviceName,
 																   const FIngestCaptureData& InIngestCaptureData);
-
-	static void ConvertPathsToFull(const FString& InTakeStoragePath, FIngestCaptureData& OutIngestCaptureData);
 };
\ No newline at end of file
diff --git a/Engine/Plugins/VirtualProduction/CaptureManager/CaptureManagerEditor/Source/LiveLinkHubWorkerManager/Private/IngestProcess/IngestProcessData.h b/Engine/Plugins/VirtualProduction/CaptureManager/CaptureManagerEditor/Source/LiveLinkHubWorkerManager/Private/IngestProcess/IngestProcessData.h
--- a/Engine/Plugins/VirtualProduction/CaptureManager/CaptureManagerEditor/Source/LiveLinkHubWorkerManager/Private/IngestProcess/IngestProcessData.h
+++ b/Engine/Plugins/VirtualProduction/CaptureManager/CaptureManagerEditor/Source/LiveLinkHubWorkerManager/Private/IngestProcess/IngestProcessData.h
@@ -2,19 +2,11 @@
 
 #pragma once
 
-#include "IngestAssetCreator.h"
-
-struct FCaptureDataTakeInfo
-{
-	FString Name;
-	double FrameRate = 0.0;
-	FIntPoint Resolution;
-	FString DeviceModel;
-};
+#include "CaptureManagerIngestPreparation.h"
 
 struct FIngestProcessResult
 {
 	FString TakeIngestPackagePath;
-	FCaptureDataTakeInfo CaptureDataTakeInfo;
+	UE::CaptureManager::FCaptureDataTakeInfo CaptureDataTakeInfo;
 	TArray<UE::CaptureManager::FCreateAssetsData> AssetsData;
 };
\ No newline at end of file
diff --git a/Engine/Plugins/VirtualProduction/CaptureManager/CaptureManagerEditor/Source/LiveLinkHubWorkerManager/Private/LiveLinkHubWorker.cpp b/Engine/Plugins/VirtualProduction/CaptureManager/CaptureManagerEditor/Source/LiveLinkHubWorkerManager/Private/LiveLinkHubWorker.cpp
--- a/Engine/Plugins/VirtualProduction/CaptureManager/CaptureManagerEditor/Source/LiveLinkHubWorkerManager/Private/LiveLinkHubWorker.cpp
+++ b/Engine/Plugins/VirtualProduction/CaptureManager/CaptureManagerEditor/Source/LiveLinkHubWorkerManager/Private/LiveLinkHubWorker.cpp
@@ -8,9 +8,6 @@
 
 #include "MessageEndpointBuilder.h"
 
-#include "CaptureData.h"
-#include "ImageSequenceTimecodeUtils.h"
-
 #include "Misc/SecureHash.h"
 #include "Misc/FileHelper.h"
 
@@ -21,11 +18,9 @@
 #include "Async/HelperFunctions.h"
 #include "Network/TcpServer.h"
 
-#include "IngestAssetCreator.h"
+#include "CaptureManagerIngestPreparation.h"
 #include "IngestProcess/IngestProcessData.h"
 
-#include "SoundWaveTimecodeUtils.h"
-
 #include "Editor.h"
 #include "Engine/Engine.h"
 #include "GlobalNamingTokens.h"
@@ -35,11 +30,6 @@
 #include "Settings/CaptureManagerEditorSettings.h"
 #include "Settings/CaptureManagerEditorTemplateTokens.h"
 
-#include "AssetToolsModule.h"
-#include "AssetRegistry/AssetRegistryModule.h"
-#include "AssetRegistry/IAssetRegistry.h"
-#include "FileHelpers.h"
-
 #define LOCTEXT_NAMESPACE "LiveLinkHubWorker"
 
 DEFINE_LOG_CATEGORY(LogLiveLinkHubWorker);
@@ -341,15 +331,14 @@ void FLiveLinkHubImportWorker::SpawnIngestTask(TSharedPtr<FEditorMessenger> InMe
 							return;
 						}
 
-						FCaptureDataAssetInfo CaptureDataAssetInfo = CaptureDataAssetInfos[0];
-						bool Created = CreateCaptureAsset(IngestProcessResult.TakeIngestPackagePath, CaptureDataAssetInfo, IngestProcessResult.CaptureDataTakeInfo);
-
-						if (GetDefault<UCaptureManagerEditorSettings>()->bAutoSaveAssets)
-						{
-							SaveCaptureCreatedAssets(IngestProcessResult.TakeIngestPackagePath);
-						}
+						const bool bAutoSave = GetDefault<UCaptureManagerEditorSettings>()->bAutoSaveAssets;
+						UFootageCaptureData* CaptureData = CreateFootageCaptureDataAsset_GameThread(
+							IngestProcessResult.TakeIngestPackagePath,
+							MoveTemp(CaptureDataAssetInfos[0]),
+							IngestProcessResult.CaptureDataTakeInfo,
+							bAutoSave);
 
-						if (Created)
+						if (CaptureData)
 						{
 							PublishDoneEvent(SharedMessenger, CaptureSourceId, TakeUploadId);
 						}
@@ -366,86 +355,6 @@ void FLiveLinkHubImportWorker::SpawnIngestTask(TSharedPtr<FEditorMessenger> InMe
 		});
 }
 
-bool FLiveLinkHubImportWorker::CreateCaptureAsset(const FString& InAssetPath, const UE::CaptureManager::FCaptureDataAssetInfo& InResult, const FCaptureDataTakeInfo& InTakeInfo)
-{
-	using namespace UE::CaptureManager;
-
-	FString CaptureDataName = InTakeInfo.Name;
-	CaptureDataName.TrimStartAndEndInline();
-	CaptureDataName.ReplaceCharInline(' ', '_');
-
-	UFootageCaptureData* CaptureData = FIngestAssetCreator::GetAssetIfExists<UFootageCaptureData>(InAssetPath, CaptureDataName);
-	if (!CaptureData)
-	{
-		CaptureData = FIngestAssetCreator::CreateAsset<UFootageCaptureData>(InAssetPath, CaptureDataName);
-		if (CaptureData)
-		{
-			CaptureData->ImageSequences.Reset();
-			CaptureData->DepthSequences.Reset();
-			CaptureData->CameraCalibrations.Reset();
-			CaptureData->AudioTracks.Reset();
-
-			for (const FCaptureDataAssetInfo::FImageSequence& ImageSequence : InResult.ImageSequences)
-			{
-				// StartTimecode and FrameRate get set during the creation of the ImageSequence and we don't want to override it
-				// UImageSequenceTimecodeUtils::SetTimecodeInfo(ImageSequence.Timecode, ImageSequence.TimecodeRate, ImageSequence.Asset);
-				CaptureData->ImageSequences.Add(ImageSequence.Asset);
-			}
-
-			for (const FCaptureDataAssetInfo::FImageSequence& DepthSequence : InResult.DepthSequences)
-			{
-				// StartTimecode and FrameRate get set during the creation of the ImageSequence and we don't want to override it
-				// UImageSequenceTimecodeUtils::SetTimecodeInfo(DepthSequence.Timecode, DepthSequence.TimecodeRate, DepthSequence.Asset);
-				CaptureData->DepthSequences.Add(DepthSequence.Asset);
-			}
-
-			for (const FCaptureDataAssetInfo::FAudio& Audio : InResult.Audios)
-			{
-				CaptureData->AudioTracks.Add(Audio.Asset);
-			}
-
-			for (const FCaptureDataAssetInfo::FCalibration& Calibration : InResult.Calibrations)
-			{
-				CaptureData->CameraCalibrations.Add(Calibration.Asset);
-			}
-
-			CaptureData->Metadata.FrameRate = InTakeInfo.FrameRate;
-			CaptureData->Metadata.DeviceModelName = InTakeInfo.DeviceModel;
-			CaptureData->Metadata.SetDeviceClass(InTakeInfo.DeviceModel);
-			CaptureData->CaptureExcludedFrames = InResult.CaptureExcludedFrames;
-
-			return true;
-		}
-	}
-
-	return false;
-}
-
-void FLiveLinkHubImportWorker::SaveCaptureCreatedAssets(const FString& InAssetPath)
-{
-	IAssetTools& AssetTools = FModuleManager::GetModuleChecked<FAssetToolsModule>(TEXT("AssetTools")).Get();
-	IAssetRegistry& AssetRegistry = FModuleManager::GetModuleChecked<FAssetRegistryModule>(TEXT("AssetRegistry")).Get();
-
-	TArray<FAssetData> AssetsData;
-	AssetRegistry.GetAssetsByPath(FName{ *InAssetPath }, AssetsData, true, false);
-
-	if (AssetsData.IsEmpty())
-	{
-		return;
-	}
-
-	TArray<UPackage*> Packages;
-	for (const FAssetData& AssetData : AssetsData)
-	{
-		UPackage* Package = AssetData.GetAsset()->GetPackage();
-		if (!Packages.Contains(Package))
-		{
-			Packages.Add(Package);
-		}
-	}
-
-	UEditorLoadingAndSavingUtils::SavePackages(Packages, true);
-}
 }
 
 #undef LOCTEXT_NAMESPACE
diff --git a/Engine/Plugins/VirtualProduction/CaptureManager/CaptureManagerEditor/Source/LiveLinkHubWorkerManager/Public/LiveLinkHubWorker.h b/Engine/Plugins/VirtualProduction/CaptureManager/CaptureManagerEditor/Source/LiveLinkHubWorkerManager/Public/LiveLinkHubWorker.h
--- a/Engine/Plugins/VirtualProduction/CaptureManager/CaptureManagerEditor/Source/LiveLinkHubWorkerManager/Public/LiveLinkHubWorker.h
+++ b/Engine/Plugins/VirtualProduction/CaptureManager/CaptureManagerEditor/Source/LiveLinkHubWorkerManager/Public/LiveLinkHubWorker.h
@@ -17,13 +17,7 @@
 
 #include "Async/TaskProgress.h"
 
-struct FCaptureDataTakeInfo;
-
-namespace UE::CaptureManager
-{
-struct FCaptureDataAssetInfo;
-
-namespace Private
+namespace UE::CaptureManager::Private
 {
 
 class FLiveLinkHubImportWorker : public TSharedFromThis<FLiveLinkHubImportWorker>
@@ -66,14 +60,7 @@ private:
 		const FGuid& InTakeUploadId,
 		TSharedPtr<UE::CaptureManager::FTaskProgress> InTaskProgress);
 
-	static bool CreateCaptureAsset(const FString& InAssetPath,
-		const UE::CaptureManager::FCaptureDataAssetInfo& InResult,
-		const FCaptureDataTakeInfo& InTakeInfo);
-
-	static void SaveCaptureCreatedAssets(const FString& InAssetPath);
-
 	TWeakPtr<FEditorMessenger> Messenger;
 };
 
-}
 }
\ No newline at end of file
